Moving Landmass plugin out of \NotForLicensees\ and into \Experimental.

diff --git a/Engine/Plugins/Experimental/Landmass/Landmass.uplugin b/Engine/Plugins/Experimental/Landmass/Landmass.uplugin
new file mode 100644
--- /dev/null
+++ b/Engine/Plugins/Experimental/Landmass/Landmass.uplugin
@@ -0,0 +1,23 @@
+{
+	"FileVersion": 3,
+	"Version": 1,
+	"VersionName": "1.0",
+	"FriendlyName": "Landmass",
+	"Description": "",
+	"Category": "Other",
+	"CreatedBy": "Epic Games",
+	"CreatedByURL": "",
+	"DocsURL": "",
+	"MarketplaceURL": "",
+	"SupportURL": "",
+	"CanContainContent": true,
+	"IsBetaVersion": true,
+	"Installed": false,
+	"Modules": [
+		{
+			"Name": "Landmass",
+			"Type": "Editor",
+			"LoadingPhase": "PreDefault"
+		}
+	]
+}
\ No newline at end of file
diff --git a/Engine/Plugins/Experimental/Landmass/Source/Runtime/Landmass.Build.cs b/Engine/Plugins/Experimental/Landmass/Source/Runtime/Landmass.Build.cs
new file mode 100644
--- /dev/null
+++ b/Engine/Plugins/Experimental/Landmass/Source/Runtime/Landmass.Build.cs
@@ -0,0 +1,26 @@
+// Copyright 1998-2019 Epic Games, Inc. All Rights Reserved.
+
+using UnrealBuildTool;
+
+public class Landmass : ModuleRules
+{
+	public Landmass(ReadOnlyTargetRules Target) : base(Target)
+	{
+		PCHUsage = ModuleRules.PCHUsageMode.UseExplicitOrSharedPCHs;
+		
+		PrivateIncludePaths.Add("Runtime/Private");
+
+		PublicDependencyModuleNames.AddRange(
+			new string[] {
+				"Core",
+				"CoreUObject",
+				"Engine",
+			}
+		);
+
+		PrivateDependencyModuleNames.AddRange(
+			new string[] {
+			}
+		);
+	}
+}
\ No newline at end of file
diff --git a/Engine/Plugins/Experimental/Landmass/Source/Runtime/Private/BrushEffectsList.cpp b/Engine/Plugins/Experimental/Landmass/Source/Runtime/Private/BrushEffectsList.cpp
new file mode 100644
--- /dev/null
+++ b/Engine/Plugins/Experimental/Landmass/Source/Runtime/Private/BrushEffectsList.cpp
@@ -0,0 +1,8 @@
+// Copyright 1998-2019 Epic Games, Inc. All Rights Reserved.
+
+#include "BrushEffectsList.h"
+
+FLandmassBrushEffectsList::FLandmassBrushEffectsList()
+{
+
+}
\ No newline at end of file
diff --git a/Engine/Plugins/Experimental/Landmass/Source/Runtime/Private/FalloffSettings.cpp b/Engine/Plugins/Experimental/Landmass/Source/Runtime/Private/FalloffSettings.cpp
new file mode 100644
--- /dev/null
+++ b/Engine/Plugins/Experimental/Landmass/Source/Runtime/Private/FalloffSettings.cpp
@@ -0,0 +1,13 @@
+// Copyright 1998-2019 Epic Games, Inc. All Rights Reserved.
+
+#include "FalloffSettings.h"
+
+FLandmassFalloffSettings::FLandmassFalloffSettings()
+	: FalloffMode(EBrushFalloffMode::Angle)
+	, FalloffAngle(45.0f)
+	, FalloffWidth(1024.0f)
+	, EdgeOffset(0)
+	, ZOffset(0)
+{
+
+}
diff --git a/Engine/Plugins/Experimental/Landmass/Source/Runtime/Private/LandmassModule.cpp b/Engine/Plugins/Experimental/Landmass/Source/Runtime/Private/LandmassModule.cpp
new file mode 100644
--- /dev/null
+++ b/Engine/Plugins/Experimental/Landmass/Source/Runtime/Private/LandmassModule.cpp
@@ -0,0 +1,22 @@
+// Copyright 1998-2019 Epic Games, Inc. All Rights Reserved.
+
+#include "LandmassModule.h"
+#include "Modules/ModuleManager.h"
+
+
+class FLandmassModule : public ILandmassModuleInterface
+{
+public:
+	virtual void StartupModule() override
+	{
+	}
+
+	virtual void ShutdownModule() override
+	{
+	}
+};
+
+//////////////////////////////////////////////////////////////////////////
+
+IMPLEMENT_MODULE(FLandmassModule, Landmass);
+
diff --git a/Engine/Plugins/Experimental/Landmass/Source/Runtime/Private/TerrainCarvingSettings.cpp b/Engine/Plugins/Experimental/Landmass/Source/Runtime/Private/TerrainCarvingSettings.cpp
new file mode 100644
--- /dev/null
+++ b/Engine/Plugins/Experimental/Landmass/Source/Runtime/Private/TerrainCarvingSettings.cpp
@@ -0,0 +1,12 @@
+// Copyright 1998-2019 Epic Games, Inc. All Rights Reserved.
+
+#include "TerrainCarvingSettings.h"
+
+FLandmassTerrainCarvingSettings::FLandmassTerrainCarvingSettings()
+	: BlendMode(EBrushBlendType::AlphaBlend)
+	, bInvertShape(false)
+	, Priority(0)
+{
+
+
+}
diff --git a/Engine/Plugins/Experimental/Landmass/Source/Runtime/Public/BrushEffectsList.h b/Engine/Plugins/Experimental/Landmass/Source/Runtime/Public/BrushEffectsList.h
new file mode 100644
--- /dev/null
+++ b/Engine/Plugins/Experimental/Landmass/Source/Runtime/Public/BrushEffectsList.h
@@ -0,0 +1,182 @@
+// Copyright 1998-2019 Epic Games, Inc. All Rights Reserved.
+
+#pragma once
+
+#include "CoreMinimal.h"
+#include "UObject/ObjectMacros.h"
+#include "BrushEffectsList.generated.h"
+
+class UTexture2D;
+class UCurveFloat;
+
+USTRUCT(BlueprintType)
+struct FBrushEffectBlurring
+{
+	GENERATED_BODY()
+
+	FBrushEffectBlurring()
+		: bBlurShape(true)
+		, Radius(2)
+	{}
+
+	UPROPERTY(EditAnywhere, BlueprintReadWrite, Category = BrushEffects)
+	bool bBlurShape;
+
+	UPROPERTY(EditAnywhere, BlueprintReadWrite, Category = BrushEffects)
+	int32 Radius;
+};
+
+
+USTRUCT(BlueprintType)
+struct FBrushEffectCurlNoise
+{
+	GENERATED_BODY()
+
+	FBrushEffectCurlNoise()
+		: Curl1Amount(0)
+		, Curl2Amount(0)
+		, Curl1Tiling(16.0)
+		, Curl2Tiling(3.0)
+	{}
+
+	UPROPERTY(EditAnywhere, BlueprintReadWrite, Category = BrushEffects)
+	float Curl1Amount;
+
+	UPROPERTY(EditAnywhere, BlueprintReadWrite, Category = BrushEffects)
+	float Curl2Amount;
+
+	UPROPERTY(EditAnywhere, BlueprintReadWrite, Category = BrushEffects)
+	float Curl1Tiling;
+
+	UPROPERTY(EditAnywhere, BlueprintReadWrite, Category = BrushEffects)
+	float Curl2Tiling;
+};
+
+USTRUCT(BlueprintType)
+struct FBrushEffectCurves
+{
+	GENERATED_BODY()
+
+	FBrushEffectCurves()
+		: bUseCurveChannel(true)
+		, ElevationCurveAsset(nullptr)
+		, ChannelEdgeOffset(0.0f)
+		, ChannelDepth(0.0f)
+		, CurveRampWidth(512.0f)
+	{}
+
+	UPROPERTY(EditAnywhere, BlueprintReadWrite, Category = BrushEffects)
+	bool bUseCurveChannel;
+
+	UPROPERTY(EditAnywhere, BlueprintReadWrite, Category = BrushEffects)
+	UCurveFloat* ElevationCurveAsset;
+
+	UPROPERTY(EditAnywhere, BlueprintReadWrite, Category = BrushEffects)
+	float ChannelEdgeOffset;
+
+	UPROPERTY(EditAnywhere, BlueprintReadWrite, Category = BrushEffects)
+	float ChannelDepth;
+
+	UPROPERTY(EditAnywhere, BlueprintReadWrite, Category = BrushEffects)
+	float CurveRampWidth;
+};
+
+
+USTRUCT(BlueprintType)
+struct FBrushEffectDisplacement
+{
+	GENERATED_BODY()
+
+	FBrushEffectDisplacement()
+		: DisplacementHeight(0)
+		, DisplacementTiling(0)
+		, Texture(nullptr)
+		, Midpoint(-128.0f)
+		, Channel(0,0,0,1)
+		, WeightmapInfluence(0.0f)
+	{}
+
+	UPROPERTY(EditAnywhere, BlueprintReadWrite, Category = BrushEffects)
+	float DisplacementHeight;
+
+	UPROPERTY(EditAnywhere, BlueprintReadWrite, Category = BrushEffects)
+	float DisplacementTiling;
+
+	UPROPERTY(EditAnywhere, BlueprintReadWrite, Category = BrushEffects)
+	UTexture2D* Texture;
+
+	UPROPERTY(EditAnywhere, BlueprintReadWrite, Category = BrushEffects)
+	float Midpoint;
+
+	UPROPERTY(EditAnywhere, BlueprintReadWrite, Category = BrushEffects)
+	FLinearColor Channel;
+
+	UPROPERTY(EditAnywhere, BlueprintReadWrite, Category = BrushEffects)
+	float WeightmapInfluence;
+};
+
+
+USTRUCT(BlueprintType)
+struct FBrushEffectSmoothBlending
+{
+	GENERATED_BODY()
+
+	UPROPERTY(EditAnywhere, BlueprintReadWrite, Category = BrushEffects)
+	float InnerSmoothDistance;
+
+	UPROPERTY(EditAnywhere, BlueprintReadWrite, Category = BrushEffects)
+	float OuterSmoothDistance;
+};
+
+USTRUCT(BlueprintType)
+struct FBrushEffectTerracing
+{
+	GENERATED_BODY()
+
+	FBrushEffectTerracing()
+		: TerraceAlpha(0.0f)
+		, TerraceSpacing(256.0f)
+		, TerraceSmoothness(0.0f)
+		, MaskLength(0.0f)
+		, MaskStartOffset(0.0f)
+	{}
+
+	UPROPERTY(EditAnywhere, BlueprintReadWrite, Category = BrushEffects)
+	float TerraceAlpha;
+
+	UPROPERTY(EditAnywhere, BlueprintReadWrite, Category = BrushEffects)
+	float TerraceSpacing;
+
+	UPROPERTY(EditAnywhere, BlueprintReadWrite, Category = BrushEffects)
+	float TerraceSmoothness;
+
+	UPROPERTY(EditAnywhere, BlueprintReadWrite, Category = BrushEffects)
+	float MaskLength;
+
+	UPROPERTY(EditAnywhere, BlueprintReadWrite, Category = BrushEffects)
+	float MaskStartOffset;
+};
+
+
+USTRUCT(BlueprintType)
+struct FLandmassBrushEffectsList
+{
+	GENERATED_BODY()
+
+	FLandmassBrushEffectsList();
+
+	UPROPERTY(EditAnywhere, BlueprintReadWrite, Category = BrushEffects)
+	FBrushEffectBlurring Blurring;
+
+	UPROPERTY(EditAnywhere, BlueprintReadWrite, Category = FalloffSettings)
+	FBrushEffectCurlNoise CurlNoise;
+
+	UPROPERTY(EditAnywhere, BlueprintReadWrite, Category = FalloffSettings)
+	FBrushEffectDisplacement Displacement;
+
+	UPROPERTY(EditAnywhere, BlueprintReadWrite, Category = FalloffSettings)
+	FBrushEffectSmoothBlending SmoothBlending;
+
+	UPROPERTY(EditAnywhere, BlueprintReadWrite, Category = FalloffSettings)
+	FBrushEffectTerracing Terracing;
+};
\ No newline at end of file
diff --git a/Engine/Plugins/Experimental/Landmass/Source/Runtime/Public/FalloffSettings.h b/Engine/Plugins/Experimental/Landmass/Source/Runtime/Public/FalloffSettings.h
new file mode 100644
--- /dev/null
+++ b/Engine/Plugins/Experimental/Landmass/Source/Runtime/Public/FalloffSettings.h
@@ -0,0 +1,36 @@
+// Copyright 1998-2019 Epic Games, Inc. All Rights Reserved.
+
+#pragma once
+
+#include "CoreMinimal.h"
+#include "UObject/ObjectMacros.h"
+#include "FalloffSettings.generated.h"
+
+UENUM(BlueprintType)
+enum class EBrushFalloffMode : uint8
+{
+	Angle,
+	Width,
+};
+
+USTRUCT(BlueprintType)
+struct FLandmassFalloffSettings
+{
+	GENERATED_BODY()
+	FLandmassFalloffSettings();
+
+	UPROPERTY(EditAnywhere, BlueprintReadWrite, Category = FalloffSettings)
+	EBrushFalloffMode FalloffMode;
+
+	UPROPERTY(EditAnywhere, BlueprintReadWrite, Category = FalloffSettings)
+	float FalloffAngle;
+
+	UPROPERTY(EditAnywhere, BlueprintReadWrite, Category = FalloffSettings)
+	float FalloffWidth;
+
+	UPROPERTY(EditAnywhere, BlueprintReadWrite, Category = FalloffSettings)
+	float EdgeOffset;
+
+	UPROPERTY(EditAnywhere, BlueprintReadWrite, Category = FalloffSettings)
+	float ZOffset;
+};
\ No newline at end of file
diff --git a/Engine/Plugins/Experimental/Landmass/Source/Runtime/Public/LandmassModule.h b/Engine/Plugins/Experimental/Landmass/Source/Runtime/Public/LandmassModule.h
new file mode 100644
--- /dev/null
+++ b/Engine/Plugins/Experimental/Landmass/Source/Runtime/Public/LandmassModule.h
@@ -0,0 +1,12 @@
+// Copyright 1998-2019 Epic Games, Inc. All Rights Reserved.
+
+#pragma once
+
+#include "CoreMinimal.h"
+#include "Stats/Stats.h"
+#include "Modules/ModuleInterface.h"
+
+
+class ILandmassModuleInterface : public IModuleInterface
+{
+};
\ No newline at end of file
diff --git a/Engine/Plugins/Experimental/Landmass/Source/Runtime/Public/TerrainCarvingSettings.h b/Engine/Plugins/Experimental/Landmass/Source/Runtime/Public/TerrainCarvingSettings.h
new file mode 100644
--- /dev/null
+++ b/Engine/Plugins/Experimental/Landmass/Source/Runtime/Public/TerrainCarvingSettings.h
@@ -0,0 +1,47 @@
+// Copyright 1998-2019 Epic Games, Inc. All Rights Reserved.
+
+#pragma once
+
+
+#include "CoreMinimal.h"
+#include "UObject/ObjectMacros.h"
+#include "FalloffSettings.h"
+#include "BrushEffectsList.h"
+#include "TerrainCarvingSettings.generated.h"
+
+/** The blend mode changes how the brush material is applied to the terrain. */
+UENUM(BlueprintType)
+enum class EBrushBlendType : uint8
+{
+	/** Alpha Blend will affect the heightmap both upwards and downwards. */
+	AlphaBlend,
+	/** Limits the brush to only lowering the terrain. */
+	Min,
+	/** Limits the brush to only raising the terrain. */
+	Max,
+	/** Performs an additive blend, using a flat Z=0 terrain as the input. Useful when you want to preserve underlying detail or ramps. */
+	Additive,
+};
+
+USTRUCT(BlueprintType)
+struct FLandmassTerrainCarvingSettings
+{
+	GENERATED_BODY()
+
+	LANDMASS_API FLandmassTerrainCarvingSettings();
+
+	UPROPERTY(EditAnywhere, BlueprintReadWrite, Category=TerrainCarvingSettings)
+	EBrushBlendType BlendMode;
+
+	UPROPERTY(EditAnywhere, BlueprintReadWrite, Category = TerrainCarvingSettings)
+	bool bInvertShape;
+
+	UPROPERTY(EditAnywhere, BlueprintReadWrite, Category = TerrainCarvingSettings)
+	FLandmassFalloffSettings FalloffSettings;
+
+	UPROPERTY(EditAnywhere, BlueprintReadWrite, Category = TerrainCarvingSettings)
+	FLandmassBrushEffectsList Effects;
+
+	UPROPERTY(EditAnywhere, BlueprintReadWrite, Category = TerrainCarvingSettings)
+	int32 Priority;
+};
