Landmass Plugin: Copy some minor code dependencies in the plugin so that it doesn't pull on a NotForLicensees plugin. (Data still needs to be upgraded)

#rnx

diff --git a/Engine/Plugins/Experimental/Landmass/Landmass.uplugin b/Engine/Plugins/Experimental/Landmass/Landmass.uplugin
--- a/Engine/Plugins/Experimental/Landmass/Landmass.uplugin
+++ b/Engine/Plugins/Experimental/Landmass/Landmass.uplugin
@@ -17,7 +17,13 @@
 		{
 			"Name": "Landmass",
 			"Type": "Runtime",
-			"LoadingPhase": "PreDefault"
+			"LoadingPhase": "Default"
+
+		},
+		{
+			"Name": "LandmassEditor",
+			"Type": "Editor",
+			"LoadingPhase": "Default"
 		}
 	]
 }
\ No newline at end of file
diff --git a/Engine/Plugins/Experimental/Landmass/Source/Editor/LandmassEditor.Build.cs b/Engine/Plugins/Experimental/Landmass/Source/Editor/LandmassEditor.Build.cs
new file mode 100644
--- /dev/null
+++ b/Engine/Plugins/Experimental/Landmass/Source/Editor/LandmassEditor.Build.cs
@@ -0,0 +1,28 @@
+// Copyright 1998-2019 Epic Games, Inc. All Rights Reserved.
+
+using UnrealBuildTool;
+
+public class LandmassEditor : ModuleRules
+{
+	public LandmassEditor(ReadOnlyTargetRules Target) : base(Target)
+	{
+		PCHUsage = ModuleRules.PCHUsageMode.UseExplicitOrSharedPCHs;
+		
+		PrivateIncludePaths.Add("Editor/Private");
+
+		PublicDependencyModuleNames.AddRange(
+			new string[] {
+				"Core",
+				"CoreUObject",
+				"Engine",
+				"LevelEditor",
+				"UnrealEd"
+			}
+		);
+
+		PrivateDependencyModuleNames.AddRange(
+			new string[] {
+			}
+		);
+	}
+}
\ No newline at end of file
diff --git a/Engine/Plugins/Experimental/Landmass/Source/Editor/Private/LandmassActor.cpp b/Engine/Plugins/Experimental/Landmass/Source/Editor/Private/LandmassActor.cpp
new file mode 100644
--- /dev/null
+++ b/Engine/Plugins/Experimental/Landmass/Source/Editor/Private/LandmassActor.cpp
@@ -0,0 +1,69 @@
+// Copyright 1998-2019 Epic Games, Inc. All Rights Reserved.
+
+#include "LandmassActor.h"
+#include "LevelEditor.h"
+#include "Modules/ModuleManager.h"
+
+
+ALandmassActor::ALandmassActor(const FObjectInitializer& ObjectInitializer)
+	: Super(ObjectInitializer)
+{
+	USceneComponent* SceneComp = CreateDefaultSubobject<USceneComponent>(TEXT("RootComponent"));
+	RootComponent = SceneComp;
+	SceneComp->Mobility = EComponentMobility::Static;
+
+	PrimaryActorTick.bCanEverTick = true;
+	PrimaryActorTick.TickGroup = TG_DuringPhysics;
+	PrimaryActorTick.bStartWithTickEnabled = true;
+	PrimaryActorTick.SetTickFunctionEnable(true);
+
+	if (!HasAnyFlags(EObjectFlags::RF_ClassDefaultObject))
+	{
+		FLevelEditorModule& LevelEditor = FModuleManager::GetModuleChecked<FLevelEditorModule>("LevelEditor");
+		OnActorSelectionChangedHandle = LevelEditor.OnActorSelectionChanged().AddUObject(this, &ALandmassActor::HandleActorSelectionChanged);
+	}
+	bIsEditorOnlyActor = true;
+}
+
+void ALandmassActor::Tick(float DeltaSeconds)
+{
+	this->CustomTick(DeltaSeconds);
+	//UE_LOG(LogTemp, Warning, TEXT("Actor Tick was called"));
+}
+
+void ALandmassActor::CustomTick_Implementation(float DeltaSeconds)
+{
+
+}
+
+bool ALandmassActor::ShouldTickIfViewportsOnly() const
+{
+	return EditorTickIsEnabled;
+}
+
+void ALandmassActor::HandleActorSelectionChanged(const TArray<UObject*>& NewSelection, bool bForceRefresh)
+{
+	if (!IsTemplate())
+	{
+		bool bUpdateActor = false;
+		if (bWasSelected && !NewSelection.Contains(this))
+		{
+			bWasSelected = false;
+			bUpdateActor = true;
+		}
+		if (!bWasSelected && NewSelection.Contains(this))
+		{
+			bWasSelected = true;
+			bUpdateActor = true;
+		}
+		if (bUpdateActor)
+		{
+			ActorSelectionChanged(bWasSelected);
+		}
+	}
+}
+
+void ALandmassActor::ActorSelectionChanged_Implementation(bool bSelected)
+{
+
+}
diff --git a/Engine/Plugins/Experimental/Landmass/Source/Editor/Private/LandmassBPEditorExtension.cpp b/Engine/Plugins/Experimental/Landmass/Source/Editor/Private/LandmassBPEditorExtension.cpp
new file mode 100644
--- /dev/null
+++ b/Engine/Plugins/Experimental/Landmass/Source/Editor/Private/LandmassBPEditorExtension.cpp
@@ -0,0 +1,25 @@
+// Copyright 1998-2019 Epic Games, Inc. All Rights Reserved.
+
+#include "LandmassBPEditorExtension.h"
+#include "CoreMinimal.h"
+
+#include "Editor.h"
+#include "EditorViewportClient.h"
+#include "LevelEditorViewport.h"
+
+ULandmassBlueprintFunctionLibrary::ULandmassBlueprintFunctionLibrary(const FObjectInitializer& ObjectInitializer)
+	: Super(ObjectInitializer)
+{
+
+}
+
+bool ULandmassBlueprintFunctionLibrary::GetCursorWorldRay(FVector& CameraLocation, FVector& RayOrigin, FVector& RayDirection)
+{
+	FViewportCursorLocation CursorLocation = GCurrentLevelEditingViewportClient->GetCursorWorldLocationFromMousePos();
+
+	CameraLocation = CursorLocation.GetOrigin();
+	RayOrigin = CursorLocation.GetOrigin();
+	RayDirection = CursorLocation.GetDirection();
+
+	return true;
+}
\ No newline at end of file
diff --git a/Engine/Plugins/Experimental/Landmass/Source/Editor/Private/LandmassEditorModule.cpp b/Engine/Plugins/Experimental/Landmass/Source/Editor/Private/LandmassEditorModule.cpp
new file mode 100644
--- /dev/null
+++ b/Engine/Plugins/Experimental/Landmass/Source/Editor/Private/LandmassEditorModule.cpp
@@ -0,0 +1,22 @@
+// Copyright 1998-2019 Epic Games, Inc. All Rights Reserved.
+
+#include "LandmassEditorModule.h"
+#include "Modules/ModuleManager.h"
+
+
+class FLandmassEditorModule : public ILandmassEditorModuleInterface
+{
+public:
+	virtual void StartupModule() override
+	{
+	}
+
+	virtual void ShutdownModule() override
+	{
+	}
+};
+
+//////////////////////////////////////////////////////////////////////////
+
+IMPLEMENT_MODULE(FLandmassEditorModule, LandmassEditor);
+
diff --git a/Engine/Plugins/Experimental/Landmass/Source/Editor/Public/LandmassActor.h b/Engine/Plugins/Experimental/Landmass/Source/Editor/Public/LandmassActor.h
new file mode 100644
--- /dev/null
+++ b/Engine/Plugins/Experimental/Landmass/Source/Editor/Public/LandmassActor.h
@@ -0,0 +1,44 @@
+// Copyright 1998-2019 Epic Games, Inc. All Rights Reserved.
+
+#pragma once
+
+#include "CoreMinimal.h"
+#include "UObject/ObjectMacros.h"
+#include "UObject/Object.h"
+#include "UObject/GCObject.h"
+#include "Tickable.h"
+#include "GameFramework/Actor.h"
+#include "LandmassActor.generated.h"
+
+
+UCLASS(Blueprintable, hidecategories = (Replication, Input, LOD, Actor, Cooking, Rendering))
+class ALandmassActor : public AActor
+{
+	GENERATED_UCLASS_BODY()
+
+public:
+
+	UFUNCTION(BlueprintNativeEvent, CallInEditor, BlueprintCallable, Category = "Tick")
+	void CustomTick(float DeltaSeconds);
+
+	virtual bool ShouldTickIfViewportsOnly() const override;
+	virtual void Tick(float DeltaSeconds) override;
+
+	UFUNCTION(BlueprintCallable, category = "Default")
+	void SetEditorTickEnaled(bool bEnabled) { EditorTickIsEnabled = bEnabled; }
+
+	UPROPERTY()
+	bool EditorTickIsEnabled = false;
+
+	UFUNCTION(BlueprintNativeEvent, CallInEditor, BlueprintCallable, Category = "Selection")
+	void ActorSelectionChanged(bool bSelected);
+
+private:
+	bool bWasSelected = false;
+
+	FDelegateHandle OnActorSelectionChangedHandle;
+
+	/** Called when the editor selection has changed. */
+	void HandleActorSelectionChanged(const TArray<UObject*>& NewSelection, bool bForceRefresh);
+
+};
\ No newline at end of file
diff --git a/Engine/Plugins/Experimental/Landmass/Source/Editor/Public/LandmassBPEditorExtension.h b/Engine/Plugins/Experimental/Landmass/Source/Editor/Public/LandmassBPEditorExtension.h
new file mode 100644
--- /dev/null
+++ b/Engine/Plugins/Experimental/Landmass/Source/Editor/Public/LandmassBPEditorExtension.h
@@ -0,0 +1,17 @@
+// Copyright 1998-2019 Epic Games, Inc. All Rights Reserved.
+
+#pragma once
+
+#include "Kismet/BlueprintFunctionLibrary.h"
+#include "Engine/EngineTypes.h"
+#include "LandmassBPEditorExtension.generated.h"
+
+
+UCLASS()
+class ULandmassBlueprintFunctionLibrary : public UBlueprintFunctionLibrary
+{
+	GENERATED_UCLASS_BODY()
+
+	UFUNCTION(BlueprintCallable, meta = (DisplayName = "Get Cursor World Ray", Keywords = "Get Cursor World Ray", UnsafeDuringActorConstruction = "true"), Category = Rendering)
+	static bool GetCursorWorldRay(FVector& CameraLocation, FVector& RayOrigin, FVector& RayDirection);
+};
\ No newline at end of file
diff --git a/Engine/Plugins/Experimental/Landmass/Source/Editor/Public/LandmassEditorModule.h b/Engine/Plugins/Experimental/Landmass/Source/Editor/Public/LandmassEditorModule.h
new file mode 100644
--- /dev/null
+++ b/Engine/Plugins/Experimental/Landmass/Source/Editor/Public/LandmassEditorModule.h
@@ -0,0 +1,12 @@
+// Copyright 1998-2019 Epic Games, Inc. All Rights Reserved.
+
+#pragma once
+
+#include "CoreMinimal.h"
+#include "Stats/Stats.h"
+#include "Modules/ModuleInterface.h"
+
+
+class ILandmassEditorModuleInterface : public IModuleInterface
+{
+};
\ No newline at end of file
