Landmass: Fix Typo

#rnx

diff --git a/Engine/Plugins/Experimental/Landmass/Source/Editor/Public/LandmassActor.h b/Engine/Plugins/Experimental/Landmass/Source/Editor/Public/LandmassActor.h
--- a/Engine/Plugins/Experimental/Landmass/Source/Editor/Public/LandmassActor.h
+++ b/Engine/Plugins/Experimental/Landmass/Source/Editor/Public/LandmassActor.h
@@ -25,7 +25,7 @@ public:
 	virtual void Tick(float DeltaSeconds) override;
 
 	UFUNCTION(BlueprintCallable, category = "Default")
-	void SetEditorTickEnaled(bool bEnabled) { EditorTickIsEnabled = bEnabled; }
+	void SetEditorTickEnabled(bool bEnabled) { EditorTickIsEnabled = bEnabled; }
 
 	UPROPERTY()
 	bool EditorTickIsEnabled = false;
