Removed dependency of Landmass on BPEditorExtension now that the asset doesn't need it anymore

diff --git a/Engine/Plugins/Experimental/Landmass/Landmass.uplugin b/Engine/Plugins/Experimental/Landmass/Landmass.uplugin
--- a/Engine/Plugins/Experimental/Landmass/Landmass.uplugin
+++ b/Engine/Plugins/Experimental/Landmass/Landmass.uplugin
@@ -25,11 +25,5 @@
 			"Type": "Editor",
 			"LoadingPhase": "Default"
 		}
-	],
-	"Plugins": [
-		{
-			"Name": "BPEditorExtension",
-			"Enabled": true
-		}
 	]
 }
\ No newline at end of file
