Fixed crash with landmass BP utility function

diff --git a/Engine/Plugins/Experimental/Landmass/Source/Editor/Private/LandmassBPEditorExtension.cpp b/Engine/Plugins/Experimental/Landmass/Source/Editor/Private/LandmassBPEditorExtension.cpp
--- a/Engine/Plugins/Experimental/Landmass/Source/Editor/Private/LandmassBPEditorExtension.cpp
+++ b/Engine/Plugins/Experimental/Landmass/Source/Editor/Private/LandmassBPEditorExtension.cpp
@@ -14,6 +14,11 @@ ULandmassBlueprintFunctionLibrary::ULandmassBlueprintFunctionLibrary(const FObje
 
 bool ULandmassBlueprintFunctionLibrary::GetCursorWorldRay(FVector& CameraLocation, FVector& RayOrigin, FVector& RayDirection)
 {
+	if (GCurrentLevelEditingViewportClient == nullptr)
+	{
+		return false;
+	}
+
 	FViewportCursorLocation CursorLocation = GCurrentLevelEditingViewportClient->GetCursorWorldLocationFromMousePos();
 
 	CameraLocation = CursorLocation.GetOrigin();
