Landmass: Added new base class LandmassBrushManagerBase with ability for BrushManager to add itself to Edit Layers. Also new base class for LandmassErosionBrush and expanded LandmassActor with some default behavior to render into the BrushManager.

#virtualized

diff --git a/Engine/Plugins/Experimental/Landmass/Source/Editor/LandmassEditor.Build.cs b/Engine/Plugins/Experimental/Landmass/Source/Editor/LandmassEditor.Build.cs
--- a/Engine/Plugins/Experimental/Landmass/Source/Editor/LandmassEditor.Build.cs
+++ b/Engine/Plugins/Experimental/Landmass/Source/Editor/LandmassEditor.Build.cs
@@ -13,6 +13,7 @@ public class LandmassEditor : ModuleRules
 				"Core",
 				"CoreUObject",
 				"EditorFramework",
+				"Landscape",
 				"Engine",
 				"LevelEditor",
 				"UnrealEd"
diff --git a/Engine/Plugins/Experimental/Landmass/Source/Editor/Private/LandmassActor.cpp b/Engine/Plugins/Experimental/Landmass/Source/Editor/Private/LandmassActor.cpp
--- a/Engine/Plugins/Experimental/Landmass/Source/Editor/Private/LandmassActor.cpp
+++ b/Engine/Plugins/Experimental/Landmass/Source/Editor/Private/LandmassActor.cpp
@@ -1,19 +1,36 @@
 // Copyright Epic Games, Inc. All Rights Reserved.
 
 #include "LandmassActor.h"
+#include "LandmassManagerBase.h"
+#include "EngineUtils.h"
+#include "Landscape.h"
+#include "Components/StaticMeshComponent.h"
+#include "Materials/MaterialInterface.h"
+#include "Kismet/GameplayStatics.h"
 #include "LevelEditor.h"
 #include "Modules/ModuleManager.h"
 
 #include UE_INLINE_GENERATED_CPP_BY_NAME(LandmassActor)
 
-
-ALandmassActor::ALandmassActor(const FObjectInitializer& ObjectInitializer)
-	: Super(ObjectInitializer)
+ALandmassActor::ALandmassActor()
 {
 	USceneComponent* SceneComp = CreateDefaultSubobject<USceneComponent>(TEXT("RootComponent"));
 	RootComponent = SceneComp;
 	SceneComp->Mobility = EComponentMobility::Static;
 
+	MeshExtentsQuad = CreateDefaultSubobject<UStaticMeshComponent>(TEXT("ExtentsMesh"));
+	MeshExtentsQuad->SetupAttachment(RootComponent);
+	MeshExtentsQuad->SetStaticMesh(Cast< UStaticMesh >(FSoftObjectPath(TEXT("StaticMesh'/Engine/ArtTools/RenderToTexture/Meshes/S_1_Unit_Plane.S_1_Unit_Plane'")).TryLoad()));
+	MeshExtentsQuad->SetMaterial(0, Cast< UMaterialInterface >(FSoftObjectPath(TEXT("Material'/Landmass/Landscape/BlueprintBrushes/Materials/Internal/BrushBounds.BrushBounds'")).TryLoad()));
+	MeshExtentsQuad->SetCollisionEnabled(ECollisionEnabled::NoCollision);
+
+	BrushSpriteMesh = CreateDefaultSubobject<UStaticMeshComponent>(TEXT("BrushSpriteMesh"));
+	BrushSpriteMesh->SetupAttachment(RootComponent);
+	BrushSpriteMesh->SetStaticMesh(Cast< UStaticMesh >(FSoftObjectPath(TEXT("StaticMesh'/Landmass/Landscape/Meshes/S_SpritePlane.S_SpritePlane'")).TryLoad()));
+	BrushSpriteMesh->SetMaterial(0, Cast< UMaterialInterface >(FSoftObjectPath(TEXT("Material'/Landmass/Landscape/BlueprintBrushes/Materials/Sprite/M_BrushSprite_01.M_BrushSprite_01'")).TryLoad()));
+	BrushSpriteMesh->SetCollisionEnabled(ECollisionEnabled::NoCollision);
+	BrushSpriteMesh->SetCastShadow(false);
+
 	PrimaryActorTick.bCanEverTick = true;
 	PrimaryActorTick.TickGroup = TG_DuringPhysics;
 	PrimaryActorTick.bStartWithTickEnabled = true;
@@ -26,10 +43,25 @@ ALandmassActor::ALandmassActor(const FObjectInitializer& ObjectInitializer)
 	}
 }
 
+void ALandmassActor::OnConstruction(const FTransform& Transform) 
+{
+	Super::OnConstruction(Transform);
+	UpdateBrushExtents();
+	FindOrSpawnManager();
+	if (BrushManager != nullptr)
+	{
+		if (AffectsHeightmap || AffectsWeightmaps || AffectsVisibility)
+		{
+			BrushManager->RequestUpdateFromBrush(this);
+		}
+	}
+	//Used for Erosion Brush to update when using Brush as a mask
+	OnBrushUpdated.Broadcast();
+}
+
 void ALandmassActor::Tick(float DeltaSeconds)
 {
 	this->CustomTick(DeltaSeconds);
-	//UE_LOG(LogTemp, Warning, TEXT("Actor Tick was called"));
 }
 
 void ALandmassActor::CustomTick_Implementation(float DeltaSeconds)
@@ -42,6 +74,134 @@ bool ALandmassActor::ShouldTickIfViewportsOnly() const
 	return EditorTickIsEnabled;
 }
 
+//Used so that the BPImplementableEvent "RenderLayer" does not have to be marked BlueprintCallable which adds the unwanted "Parent:FunctionCall" to the BP graph
+void ALandmassActor::RenderLayer_Native(const FLandscapeBrushParameters& InParameters)
+{
+	BrushRenderParameters = InParameters;
+	RenderLayer(InParameters);
+}
+
+void ALandmassActor::RenderLayer_Implementation(const FLandscapeBrushParameters& InParameters)
+{
+
+}
+
+void ALandmassActor::UpdateBrushExtents()
+{
+	//Mesh Quad represents the Brush AABB, update its size and use it to calculate the extents
+	MeshExtentsQuad->SetRelativeScale3D(FVector(BrushSize, BrushSize, BrushSize));
+	FVector2D Origin = FVector2D(MeshExtentsQuad->Bounds.Origin.X, MeshExtentsQuad->Bounds.Origin.Y);
+	FVector2D BoxExtent = FVector2D(MeshExtentsQuad->Bounds.BoxExtent.X, MeshExtentsQuad->Bounds.BoxExtent.Y);
+
+	BrushExtents = FVector4(Origin - BoxExtent, Origin + BoxExtent);
+
+	float SpriteMeshScale = (BrushSize)  / 8;
+	BrushSpriteMesh->SetRelativeScale3D(FVector(SpriteMeshScale, SpriteMeshScale, SpriteMeshScale));
+
+	if (DrawToEntireLandscape)
+	{
+
+	}
+
+}
+
+void ALandmassActor::FastPreviewMode()
+{
+	BrushManager->TogglePreviewMode(true);
+}
+
+void ALandmassActor::RestoreLandscapeEditing()
+{
+	BrushManager->TogglePreviewMode(false);
+}
+
+void ALandmassActor::MoveBrushUp()
+{
+	if (BrushManager != nullptr)
+	{
+		BrushManager->MoveBrushUp(this);
+	}
+}
+
+void ALandmassActor::MoveBrushDown()
+{
+	if (BrushManager != nullptr)
+	{
+		BrushManager->MoveBrushDown(this);
+	}
+}
+
+void ALandmassActor::MoveToTop()
+{
+	if (BrushManager != nullptr)
+	{
+		BrushManager->MoveBrushToTop(this);
+	}
+}
+
+void ALandmassActor::MoveToBottom()
+{
+	if (BrushManager != nullptr)
+	{
+		BrushManager->MoveBrushToBottom(this);
+	}
+}
+
+void ALandmassActor::FindOrSpawnManager()
+{
+#if WITH_EDITOR
+	//If brush already has a manager, do nothing
+	if (BrushManager != nullptr || this->HasAnyFlags(RF_Transient))
+	{
+		return;
+	}
+
+	TArray<AActor*> FoundActors;
+	UGameplayStatics::GetAllActorsOfClass(Cast<UObject>(GetWorld()), ALandmassManagerBase::StaticClass(), FoundActors);
+
+	if (FoundActors.Num() > 0)
+	{
+		BrushManager = Cast<ALandmassManagerBase>(FoundActors[0]);
+	}
+	else
+	{
+		FActorSpawnParameters SpawnParams;
+		SpawnParams.bAllowDuringConstructionScript = true; 
+
+		FName BlueprintPath = TEXT("Class'/Landmass/Landscape/BlueprintBrushes/LandmassBrushManager.LandmassBrushManager_C'");
+		UClass* ManagerBPClass = Cast<UClass>(StaticLoadObject(UClass::StaticClass(), NULL, *BlueprintPath.ToString()));
+		ALandmassManagerBase* SpawnedManager = GetWorld()->SpawnActor<ALandmassManagerBase>(ManagerBPClass, SpawnParams);
+
+		BrushManager = SpawnedManager;
+	}
+
+	for (TActorIterator<ALandscape> LandscapeIterator(GetWorld()); LandscapeIterator; ++LandscapeIterator)
+	{
+		if (LandscapeIterator->CanHaveLayersContent())
+		{
+			BrushManager->SetTargetLandscape(*LandscapeIterator);
+			break;
+		}
+	}
+#endif
+}
+
+void ALandmassActor::DrawBrushMaterial(UMaterialInterface* InMaterial)
+{
+	if (BrushManager != nullptr)
+	{
+		BrushManager->DrawBrushMaterial(this, InMaterial);
+	}
+}
+
+void ALandmassActor::SetMeshExentsMaterial(UMaterialInterface* Material)
+{
+	if (Material != nullptr)
+	{
+		MeshExtentsQuad->SetMaterial(0, Material);
+	}
+}
+
 void ALandmassActor::HandleActorSelectionChanged(const TArray<UObject*>& NewSelection, bool bForceRefresh)
 {
 	if (!IsTemplate())
@@ -60,6 +220,7 @@ void ALandmassActor::HandleActorSelectionChanged(const TArray<UObject*>& NewSele
 		if (bUpdateActor)
 		{
 			ActorSelectionChanged(bWasSelected);
+			MeshExtentsQuad->SetVisibility(bWasSelected);
 		}
 	}
 }
@@ -67,5 +228,4 @@ void ALandmassActor::HandleActorSelectionChanged(const TArray<UObject*>& NewSele
 void ALandmassActor::ActorSelectionChanged_Implementation(bool bSelected)
 {
 
-}
-
+}
\ No newline at end of file
diff --git a/Engine/Plugins/Experimental/Landmass/Source/Editor/Private/LandmassBPEditorExtension.cpp b/Engine/Plugins/Experimental/Landmass/Source/Editor/Private/LandmassBPEditorExtension.cpp
--- a/Engine/Plugins/Experimental/Landmass/Source/Editor/Private/LandmassBPEditorExtension.cpp
+++ b/Engine/Plugins/Experimental/Landmass/Source/Editor/Private/LandmassBPEditorExtension.cpp
@@ -3,15 +3,10 @@
 #include "LandmassBPEditorExtension.h"
 
 #include "LevelEditorViewport.h"
+#include "Engine/Texture.h"
 
 #include UE_INLINE_GENERATED_CPP_BY_NAME(LandmassBPEditorExtension)
 
-ULandmassBlueprintFunctionLibrary::ULandmassBlueprintFunctionLibrary(const FObjectInitializer& ObjectInitializer)
-	: Super(ObjectInitializer)
-{
-
-}
-
 bool ULandmassBlueprintFunctionLibrary::GetCursorWorldRay(FVector& CameraLocation, FVector& RayOrigin, FVector& RayDirection)
 {
 	if (GCurrentLevelEditingViewportClient == nullptr)
@@ -27,3 +22,51 @@ bool ULandmassBlueprintFunctionLibrary::GetCursorWorldRay(FVector& CameraLocatio
 
 	return true;
 }
+
+//Function that returns only the overlapping region between two world AABB extents
+void ULandmassBlueprintFunctionLibrary::CombineWorldExtents(FVector4 ExtentsA, FVector4 ExtentsB, FVector4& CombinedExtents)
+{
+	CombinedExtents.X = FMath::Max(ExtentsA.X, ExtentsB.X);
+	CombinedExtents.Y = FMath::Max(ExtentsA.Y, ExtentsB.Y);
+	CombinedExtents.Z = FMath::Min(ExtentsA.Z, ExtentsB.Z);
+	CombinedExtents.W = FMath::Min(ExtentsA.W, ExtentsB.W);
+}
+
+//Converts World Extents into Canvas coordinates with scaled UVs to match the partial draw area
+void ULandmassBlueprintFunctionLibrary::WorldExtentsToCanvasCoordinates(FVector4 WorldExtents, FLandmassLandscapeInfo LandscapeInfo, FVector2D& ScreenPosition, FVector2D& ScreenSize, FVector2D& CoordinatePosition, FVector2D& CoordinateSize)
+{
+	FVector LandscapePosition = LandscapeInfo.LandscapeTransform.GetLocation();
+	FVector LandscapeScale = LandscapeInfo.LandscapeTransform.GetScale3D();
+	FVector2D RTWorldSize = FVector2D(LandscapeInfo.RenderTargetResolution) * FVector2D(LandscapeScale.X, LandscapeScale.Y);
+	
+	FVector2D ExtentsMin = FVector2D(WorldExtents.X, WorldExtents.Y) - FVector2D(LandscapePosition.X, LandscapePosition.Y);
+	FVector2D ExtentsMax = FVector2D(WorldExtents.Z, WorldExtents.W) - FVector2D(LandscapePosition.X, LandscapePosition.Y);
+	FVector2D ExtentSize = FVector2D(WorldExtents.Z, WorldExtents.W) - FVector2D(WorldExtents.X, WorldExtents.Y);
+
+	ScreenPosition = ExtentsMin / RTWorldSize;
+	ScreenPosition *= LandscapeInfo.RenderTargetResolution;
+	ScreenPosition /= 4;
+	ScreenPosition = FVector2D(FMath::FloorToFloat(ScreenPosition.X), FMath::FloorToFloat(ScreenPosition.Y));
+	ScreenPosition *= 4;
+
+	FVector2D ScreenPositionMax = ExtentsMax / RTWorldSize;
+	ScreenPositionMax *= LandscapeInfo.RenderTargetResolution;
+	ScreenPositionMax /= 4;
+	ScreenPositionMax = FVector2D(FMath::CeilToFloat(ScreenPositionMax.X+1), FMath::CeilToFloat(ScreenPositionMax.Y+1));
+	ScreenPositionMax *= 4;
+
+	ScreenSize = ScreenPositionMax - ScreenPosition;
+
+	CoordinatePosition = ScreenPosition / LandscapeInfo.RenderTargetResolution;
+	CoordinateSize = ScreenSize / LandscapeInfo.RenderTargetResolution;
+
+}
+
+//Used when modifying texture settings via blueprint to force the texture to update since it may not 
+void ULandmassBlueprintFunctionLibrary::ForceUpdateTexture(UTexture* InTexture)
+{
+	InTexture->Modify();
+	InTexture->MarkPackageDirty();
+	InTexture->PostEditChange();
+	InTexture->UpdateResource();
+}
\ No newline at end of file
diff --git a/Engine/Plugins/Experimental/Landmass/Source/Editor/Private/LandmassErosionBrushBase.cpp b/Engine/Plugins/Experimental/Landmass/Source/Editor/Private/LandmassErosionBrushBase.cpp
new file mode 100644
--- /dev/null
+++ b/Engine/Plugins/Experimental/Landmass/Source/Editor/Private/LandmassErosionBrushBase.cpp
@@ -0,0 +1,156 @@
+// Copyright Epic Games, Inc. All Rights Reserved.
+
+#include "LandmassErosionBrushBase.h"
+
+#include "Modules/ModuleManager.h"
+#include "EngineUtils.h"
+#include "LevelEditor.h"
+#include "Landscape.h"
+#include "LandscapeEditTypes.h"
+#include "LandscapeModule.h"
+
+#include UE_INLINE_GENERATED_CPP_BY_NAME(LandmassErosionBrushBase)
+
+DEFINE_LOG_CATEGORY(LandmassErosionBrush);
+
+ALandmassErosionBrushBase::ALandmassErosionBrushBase()
+{
+	if (!HasAnyFlags(EObjectFlags::RF_ClassDefaultObject))
+	{
+		FLevelEditorModule& LevelEditor = FModuleManager::GetModuleChecked<FLevelEditorModule>("LevelEditor");
+		OnActorSelectionChangedHandle = LevelEditor.OnActorSelectionChanged().AddUObject(this, &ALandmassErosionBrushBase::HandleActorSelectionChanged);
+	}
+}
+
+void ALandmassErosionBrushBase::OnConstruction(const FTransform& Transform)
+{
+	Super::OnConstruction(Transform);
+	FindAndAssignLandscape();
+
+}
+
+void ALandmassErosionBrushBase::SetTargetLandscape(ALandscape* InTargetLandscape)
+{
+#if WITH_EDITOR
+	if (OwningLandscape != InTargetLandscape)
+	{
+		if (OwningLandscape)
+		{
+			OwningLandscape->RemoveBrush(this);
+		}
+
+		if (!InTargetLandscape)
+		{
+			if (OwningLandscape != nullptr)
+			{
+				// This can occur if the RemoveBrush call above did not do anything because the manager
+				// was removed from the landscape in some other way (probably in landscape mode panel)
+				SetOwningLandscape(nullptr);
+			}
+			return;
+		}
+
+		if (!InTargetLandscape->CanHaveLayersContent())
+		{
+			UE_LOG(LandmassErosionBrush, Warning, TEXT("Landscape target for Landmass Erosion Brush did not have edit layers enabled. Unable to attach manager."));
+			SetOwningLandscape(nullptr);
+		}
+		else
+		{
+			static const FName PatchLayerName = FName("Landmass Erosion");
+
+			ILandscapeModule& LandscapeModule = FModuleManager::GetModuleChecked<ILandscapeModule>("Landscape");
+			int32 PatchLayerIndex = LandscapeModule.GetLandscapeEditorServices()->GetOrCreateEditLayer(PatchLayerName, InTargetLandscape);
+
+			// Among other things, this will call SetOwningLandscape on us.
+			InTargetLandscape->AddBrushToLayer(PatchLayerIndex, this);
+
+			// It's not clear whether this is really necessary, but we do it for consistency because Landscape does this in its
+			// PostLoad for all its brushes (through FLandscapeLayerBrush::SetOwner). One would think that it would be done 
+			// in AddBrushToLayer if it were at all important, but it currently isn't...
+			if (this->GetTypedOuter<ULevel>() != InTargetLandscape->GetTypedOuter<ULevel>())
+			{
+				// Change owner to be that level
+				this->Rename(nullptr, InTargetLandscape->GetTypedOuter<ULevel>());
+			}
+		}
+	}
+#endif
+}
+
+ALandscape* ALandmassErosionBrushBase::GetLandscape()
+{
+	return DetailPanelLandscape.Get();
+}
+
+void ALandmassErosionBrushBase::SetOwningLandscape(ALandscape* InOwningLandscape)
+{
+	Super::SetOwningLandscape(InOwningLandscape);
+
+	DetailPanelLandscape = OwningLandscape;
+}
+
+void ALandmassErosionBrushBase::FindAndAssignLandscape()
+{
+	//Skip if this is a transient class
+	if (this->HasAnyFlags(RF_Transient))
+	{
+		return;
+	}
+
+	for (TActorIterator<ALandscape> LandscapeIterator(GetWorld()); LandscapeIterator; ++LandscapeIterator)
+	{
+		if (LandscapeIterator->CanHaveLayersContent())
+		{
+			SetTargetLandscape(*LandscapeIterator);
+			break;
+		}
+	}
+}
+
+// We override PostEditChange to allow the users to change the owning landscape via a property displayed in the detail panel.
+void ALandmassErosionBrushBase::PostEditChangeProperty(FPropertyChangedEvent& PropertyChangedEvent)
+{
+	Super::PostEditChangeProperty(PropertyChangedEvent);
+
+	// Do a bunch of checks to make sure that we don't try to do anything when the editing is happening inside the blueprint editor.
+	UWorld* World = GetWorld();
+	if (IsTemplate() || !IsValid(this) || !IsValid(World) || World->WorldType != EWorldType::Editor)
+	{
+		return;
+	}
+
+	if (PropertyChangedEvent.Property
+		&& (PropertyChangedEvent.Property->GetFName() == GET_MEMBER_NAME_CHECKED(ALandmassErosionBrushBase, DetailPanelLandscape)))
+	{
+		SetTargetLandscape(DetailPanelLandscape.Get());
+	}
+}
+
+void ALandmassErosionBrushBase::HandleActorSelectionChanged(const TArray<UObject*>& NewSelection, bool bForceRefresh)
+{
+	if (!IsTemplate())
+	{
+		bool bUpdateActor = false;
+		if (bWasSelected && !NewSelection.Contains(this))
+		{
+			bWasSelected = false;
+			bUpdateActor = true;
+		}
+		if (!bWasSelected && NewSelection.Contains(this))
+		{
+			bWasSelected = true;
+			bUpdateActor = true;
+		}
+		if (bUpdateActor)
+		{
+			ActorSelectionChanged(bWasSelected);
+		}
+	}
+}
+
+void ALandmassErosionBrushBase::ActorSelectionChanged_Implementation(bool bSelected)
+{
+
+}
+
diff --git a/Engine/Plugins/Experimental/Landmass/Source/Editor/Private/LandmassManagerBase.cpp b/Engine/Plugins/Experimental/Landmass/Source/Editor/Private/LandmassManagerBase.cpp
new file mode 100644
--- /dev/null
+++ b/Engine/Plugins/Experimental/Landmass/Source/Editor/Private/LandmassManagerBase.cpp
@@ -0,0 +1,597 @@
+// Copyright Epic Games, Inc. All Rights Reserved.
+
+#include "LandmassManagerBase.h"
+
+#include "Modules/ModuleManager.h"
+#include "LevelEditor.h"
+#include "Landscape.h"
+#include "LandscapeEditTypes.h"
+#include "LandscapeModule.h"
+#include "Kismet/KismetMathLibrary.h"
+
+#include UE_INLINE_GENERATED_CPP_BY_NAME(LandmassManagerBase)
+
+DEFINE_LOG_CATEGORY(LandmassManager);
+
+ALandmassManagerBase::ALandmassManagerBase()
+{
+	BrushTreeDepth = 3;
+	if (!HasAnyFlags(EObjectFlags::RF_ClassDefaultObject))
+	{
+		FLevelEditorModule& LevelEditor = FModuleManager::GetModuleChecked<FLevelEditorModule>("LevelEditor");
+		OnActorSelectionChangedHandle = LevelEditor.OnActorSelectionChanged().AddUObject(this, &ALandmassManagerBase::HandleActorSelectionChanged);
+	}
+}
+
+void ALandmassManagerBase::SetTargetLandscape(ALandscape* InTargetLandscape)
+{
+#if WITH_EDITOR
+	if (OwningLandscape != InTargetLandscape)
+	{
+		if (OwningLandscape)
+		{
+			OwningLandscape->RemoveBrush(this);
+		}
+
+		if (!InTargetLandscape)
+		{
+			if (OwningLandscape != nullptr)
+			{
+				// This can occur if the RemoveBrush call above did not do anything because the manager
+				// was removed from the landscape in some other way (probably in landscape mode panel)
+				SetOwningLandscape(nullptr);
+			}
+			return;
+		}
+
+		if (!InTargetLandscape->CanHaveLayersContent())
+		{
+			UE_LOG(LandmassManager, Warning, TEXT("Landscape target for Landmass Brush manager did not have edit layers enabled. Unable to attach manager."));
+			SetOwningLandscape(nullptr);
+		}
+		else
+		{
+			static const FName PatchLayerName = FName("Landmass Brushes");
+
+			ILandscapeModule& LandscapeModule = FModuleManager::GetModuleChecked<ILandscapeModule>("Landscape");
+			int32 PatchLayerIndex = LandscapeModule.GetLandscapeEditorServices()->GetOrCreateEditLayer(PatchLayerName, InTargetLandscape);
+
+			// Among other things, this will call SetOwningLandscape on us.
+			InTargetLandscape->AddBrushToLayer(PatchLayerIndex, this);
+
+			// It's not clear whether this is really necessary, but we do it for consistency because Landscape does this in its
+			// PostLoad for all its brushes (through FLandscapeLayerBrush::SetOwner). One would think that it would be done 
+			// in AddBrushToLayer if it were at all important, but it currently isn't...
+			if (this->GetTypedOuter<ULevel>() != InTargetLandscape->GetTypedOuter<ULevel>())
+			{
+				// Change owner to be that level
+				this->Rename(nullptr, InTargetLandscape->GetTypedOuter<ULevel>());
+			}
+		}
+	}
+#endif
+}
+
+//Copied from KismetMathlibrary because the function was note exported via ENGINE_API
+int32 ALandmassManagerBase::Convert2DTo1D(const FIntPoint& Index2D, int32 XSize)
+{
+	const bool bInvalidBounds =
+		Index2D.X < 0 ||
+		Index2D.X >= XSize ||
+		Index2D.Y < 0
+		;
+
+	if (bInvalidBounds)
+	{
+		return 0;
+	}
+
+	if (XSize <= 0)
+	{
+		return 0;
+	}
+
+	return Index2D.X + (Index2D.Y * XSize);
+}
+
+void ALandmassManagerBase::PopulateNodeTree()
+{
+	BrushNodeData.Empty();
+	int32 NodeIndexOffset = 0;
+	int32 ParentIndexOffset = 0;
+
+	BrushTreeDepth = FMath::Clamp(BrushTreeDepth, 0, 4);
+
+	FVector2D WorldSize = { LandscapeInformation.LandscapeTransform.GetScale3D().X, LandscapeInformation.LandscapeTransform.GetScale3D().Y };
+	WorldSize *= LandscapeInformation.LandscapeQuads;
+
+	//For each level
+	for (int i = 0; i <= BrushTreeDepth; i++)
+	{
+		int32 SizeOfLevel = FMath::TruncToInt(FMath::Pow(2.0, i));
+		int32 NodesInLevel = SizeOfLevel * SizeOfLevel;
+		int32 SizeOfNextLevel = FMath::TruncToInt(FMath::Pow(2.0, i + 1));
+		NodeIndexOffset += NodesInLevel;
+
+		//Create empty nodes
+		for (int j = 0; j <= NodesInLevel - 1; j++)
+		{
+			FVector2D NodeWorldSize = WorldSize / SizeOfLevel;
+
+			FIntPoint NodeIndex2D = FIntPoint( j % SizeOfLevel, j / SizeOfLevel );
+			FVector LandscapePosition = LandscapeInformation.LandscapeTransform.GetLocation();
+			FVector2D NodeMinExtent = FVector2D(NodeIndex2D.X, NodeIndex2D.Y) * NodeWorldSize + FVector2D(LandscapePosition.X, LandscapePosition.Y);
+			FVector4 CurrentNodeExtents = FVector4(NodeMinExtent.X, NodeMinExtent.Y, NodeMinExtent.X + NodeWorldSize.X, NodeMinExtent.Y + NodeWorldSize.Y);
+
+			FBrushDataTree NewNode;
+			NewNode.CurrentLevel = i;
+			//Calculate the index of this node's parent by going up one level of subdivision
+			NewNode.ParentIndex = Convert2DTo1D((NodeIndex2D / FIntPoint(2, 2)), SizeOfLevel / 2) + ParentIndexOffset;
+			NewNode.Index_x0y0 = -1;
+			NewNode.Index_x1y0 = -1;
+			NewNode.Index_x0y1 = -1;
+			NewNode.Index_x1y1 = -1;
+			NewNode.ChildDataCount = 0;
+			NewNode.NodeExtents = CurrentNodeExtents;
+
+			if (i < BrushTreeDepth)
+			{
+				//Nodes contain the index of child nodes, calculate the index based on index from subdivision of level
+				NewNode.Index_x0y0 = Convert2DTo1D( (NodeIndex2D * FIntPoint(2,2)) + FIntPoint(0,0), SizeOfNextLevel) +  NodeIndexOffset;
+				NewNode.Index_x1y0 = Convert2DTo1D( (NodeIndex2D * FIntPoint(2,2)) + FIntPoint(1,0), SizeOfNextLevel) +  NodeIndexOffset;
+				NewNode.Index_x0y1 = Convert2DTo1D( (NodeIndex2D * FIntPoint(2,2)) + FIntPoint(0,1), SizeOfNextLevel) +  NodeIndexOffset;
+				NewNode.Index_x1y1 = Convert2DTo1D( (NodeIndex2D * FIntPoint(2,2)) + FIntPoint(1,1), SizeOfNextLevel) +  NodeIndexOffset;
+			}
+
+			BrushNodeData.Add(NewNode);
+
+		}
+		int32 NodesInPreviousLevel = FMath::TruncToInt(FMath::Pow(2.0, i - 1)) * FMath::TruncToInt(FMath::Pow(2.0, i - 1));
+		ParentIndexOffset += NodesInPreviousLevel;
+
+	}
+}
+
+TArray<ALandmassActor*> ALandmassManagerBase::GetActorsWithinModifiedNodes(TArray<int32>& InModifiedNodes)
+{
+	TArray<ALandmassActor*> ActorsInNodes;
+	TArray<int32> NodesToCheck = InModifiedNodes;
+	TArray<int32> ParentsToCheck = InModifiedNodes;
+	TArray<int32> NodesToCheckNext;
+	TArray<int32> ParentsToCheckNext;
+
+	//Check each level of tree
+	for (int i = 0; i <= BrushTreeDepth; i++)
+	{
+		NodesToCheckNext.Empty();
+
+		//Check all nodes in the current level
+		for (int32 CurrentNodeIndex : NodesToCheck)
+		{
+			if (!BrushNodeData[CurrentNodeIndex].BrushActors.IsEmpty())
+			{
+				//Get any brushes from those nodes
+				for (ALandmassActor* CurrentActor : BrushNodeData[CurrentNodeIndex].BrushActors)
+				{
+					ActorsInNodes.AddUnique(CurrentActor);
+				}
+			}
+
+			//If any of the node's children have data, add them to list to check at the next level
+			if (i < BrushTreeDepth && BrushNodeData[CurrentNodeIndex].ChildDataCount > 0)
+			{
+				NodesToCheckNext.Add(BrushNodeData[CurrentNodeIndex].Index_x0y0);
+				NodesToCheckNext.Add(BrushNodeData[CurrentNodeIndex].Index_x1y0);
+				NodesToCheckNext.Add(BrushNodeData[CurrentNodeIndex].Index_x0y1);
+				NodesToCheckNext.Add(BrushNodeData[CurrentNodeIndex].Index_x1y1);
+			}
+		}
+		//Set the next level iteration to check those nodes that were found with data
+		NodesToCheck = NodesToCheckNext;
+		ParentsToCheckNext.Empty();
+
+		//Also check the node's parent for data, since it is possible the search began with leaf nodes
+		for (int32 CurrentNodeIndex : ParentsToCheck)
+		{
+			FBrushDataTree CheckNode = BrushNodeData[CurrentNodeIndex];
+			int32 CurrentParentIndex = CheckNode.ParentIndex;
+
+			if (CurrentNodeIndex > 0 && CheckNode.CurrentLevel > 0)
+			{
+				if (!BrushNodeData[CurrentParentIndex].BrushActors.IsEmpty())
+				{
+					//Get all brushes from the parent node
+					for (ALandmassActor* CurrentActor : BrushNodeData[CurrentParentIndex].BrushActors)
+					{
+						ActorsInNodes.AddUnique(CurrentActor);
+					}
+				}
+
+				ParentsToCheckNext.AddUnique(CurrentParentIndex);
+			}
+		}
+		//Set the next level's iteration to check the parents of checked nodes
+		ParentsToCheck = ParentsToCheckNext;
+
+	}
+
+	return ActorsInNodes;
+}
+
+//Go through nodes and update nodes to include the number of brushes in all their children under ChildDataCount
+void ALandmassManagerBase::UpdateChildDataCounts()
+{
+	int32 CurrentActorCount;
+
+	for (FBrushDataTree& CurrentNode : BrushNodeData)
+	{
+		CurrentNode.ChildDataCount = 0;
+		CurrentActorCount = CurrentNode.BrushActors.Num();
+
+		if (CurrentActorCount > 0 && CurrentNode.CurrentLevel > 0)
+		{
+			int32 CurrentLevel = CurrentNode.CurrentLevel;
+			int32 ParentToUpdate = CurrentNode.ParentIndex;
+
+			for (int i = 0; i <= CurrentLevel - 1; i++)
+			{
+				FBrushDataTree& CurrentParentData = BrushNodeData[ParentToUpdate];
+				BrushNodeData[ParentToUpdate].ChildDataCount += CurrentActorCount;
+				ParentToUpdate = CurrentParentData.ParentIndex;
+			}
+		}
+
+	}
+}
+
+//When a list of nodes includes all the child nodes of a parent, replace the children with the parent
+void ALandmassManagerBase::ConsolidateNodes(TArray<int32>& NodesToConsolidate)
+{
+	TArray<int32> FullyOccupiedParentNodes;
+
+	//Check each tree level
+	for (int i = 0; i <= BrushTreeDepth - 1; i++)
+	{
+		FullyOccupiedParentNodes.Empty();
+
+		for (int32 CurrentCheckIndex : NodesToConsolidate)
+		{
+			//Check if Node list contains all the children of a given node
+			int32 CurrentParentIndex = BrushNodeData[CurrentCheckIndex].ParentIndex;
+			FBrushDataTree CurrentParentData = BrushNodeData[CurrentParentIndex];
+			int32 Child_x0y0 = NodesToConsolidate.Find(CurrentParentData.Index_x0y0);
+			int32 Child_x1y0 = NodesToConsolidate.Find(CurrentParentData.Index_x1y0);
+			int32 Child_x0y1 = NodesToConsolidate.Find(CurrentParentData.Index_x0y1);
+			int32 Child_x1y1 = NodesToConsolidate.Find(CurrentParentData.Index_x1y1);
+
+			//If it does, mark it as a fully occupied parent node
+			if (Child_x0y0 >= 0 && Child_x1y0 >= 0 && Child_x0y1 >= 0 && Child_x1y1 >= 0)
+			{
+				FullyOccupiedParentNodes.AddUnique(CurrentParentIndex);
+			}
+		}
+
+		//Since we added the parent node to the list and all its children were also occupied, remove the children
+		for (int32 FullyOccupiedNode : FullyOccupiedParentNodes)
+		{
+			FBrushDataTree CurrentNodeData = BrushNodeData[FullyOccupiedNode];
+			NodesToConsolidate.Remove(CurrentNodeData.Index_x0y0);
+			NodesToConsolidate.Remove(CurrentNodeData.Index_x1y0);
+			NodesToConsolidate.Remove(CurrentNodeData.Index_x0y1);
+			NodesToConsolidate.Remove(CurrentNodeData.Index_x1y1);
+			NodesToConsolidate.Insert(FullyOccupiedNode, 0);
+		}
+	}
+}
+
+TArray<int32> ALandmassManagerBase::GetNodesWithinExtents(FVector4& InExtents)
+{
+	TArray<int32> OverlappedNodes = { 0 };
+	TArray<int32> NodesToCheckNext;
+	FVector2D CheckExtentMin = { InExtents.X, InExtents.Y };
+	FVector2D CheckExtentMax = { InExtents.Z, InExtents.W };
+
+	//Check each level of tree
+	for (int i = 0; i <= BrushTreeDepth - 1; i++)
+	{
+		NodesToCheckNext.Empty();
+
+		//Check which nodes the given extents overlap, down to the smallest level
+		for (int32 CurrentCheckIndex : OverlappedNodes)
+		{
+			FVector4 CurrentNodeExtents = BrushNodeData[CurrentCheckIndex].NodeExtents;
+			FVector2D CurrentNodeCenter = (FVector2D(CurrentNodeExtents.X, CurrentNodeExtents.Y) + FVector2D(CurrentNodeExtents.Z, CurrentNodeExtents.W)) / 2;
+			if ((CheckExtentMin.X < CurrentNodeCenter.X) && (CheckExtentMin.Y < CurrentNodeCenter.Y))
+			{
+				NodesToCheckNext.Add(BrushNodeData[CurrentCheckIndex].Index_x0y0);
+			}
+			if ((CheckExtentMax.X > CurrentNodeCenter.X) && (CheckExtentMin.Y < CurrentNodeCenter.Y))
+			{
+				NodesToCheckNext.Add(BrushNodeData[CurrentCheckIndex].Index_x1y0);
+			}
+			if ((CheckExtentMin.X < CurrentNodeCenter.X) && (CheckExtentMax.Y > CurrentNodeCenter.Y))
+			{
+				NodesToCheckNext.Add(BrushNodeData[CurrentCheckIndex].Index_x0y1);
+			}
+			if ((CheckExtentMax.X > CurrentNodeCenter.X) && (CheckExtentMax.Y > CurrentNodeCenter.Y))
+			{
+				NodesToCheckNext.Add(BrushNodeData[CurrentCheckIndex].Index_x1y1);
+			}
+		}
+		OverlappedNodes = NodesToCheckNext;
+
+	}
+	//Consolidate the nodes back up to the lowest level possible. ie, if all nodes are returned, just use the root node
+	ConsolidateNodes(OverlappedNodes);
+	return OverlappedNodes;
+}
+
+//Go through and remove the Brush from the entire tree
+TArray<int32> ALandmassManagerBase::RemoveBrushFromTree(ALandmassActor* BrushToRemove)
+{
+	TArray<int32> RemovedIndices;
+	TArray<int32> NodesToCheck = {0};
+	TArray<int32> NodesToCheckNext;
+
+	for (int i = 0; i <= BrushTreeDepth; i++)
+	{
+		NodesToCheckNext.Empty();
+
+		for (int32 CurrentCheckIndex : NodesToCheck)
+		{
+			if (!BrushNodeData[CurrentCheckIndex].BrushActors.IsEmpty())
+			{
+				int32 RemovalIndex = BrushNodeData[CurrentCheckIndex].BrushActors.Find(BrushToRemove);
+				if (RemovalIndex >= 0)
+				{
+					BrushNodeData[CurrentCheckIndex].BrushActors.Remove(BrushToRemove);
+					RemovedIndices.Add(CurrentCheckIndex);
+				}
+				//Remove all null entries while at it
+				BrushNodeData[CurrentCheckIndex].BrushActors.Remove(nullptr);
+			}
+
+			if (BrushNodeData[CurrentCheckIndex].ChildDataCount > 0)
+			{
+				NodesToCheckNext.Add(BrushNodeData[CurrentCheckIndex].Index_x0y0);
+				NodesToCheckNext.Add(BrushNodeData[CurrentCheckIndex].Index_x1y0);
+				NodesToCheckNext.Add(BrushNodeData[CurrentCheckIndex].Index_x0y1);
+				NodesToCheckNext.Add(BrushNodeData[CurrentCheckIndex].Index_x1y1);
+			}
+
+		}
+
+		NodesToCheck = NodesToCheckNext;
+	}
+
+	return RemovedIndices;
+}
+
+//Sort one array to match another. Used to ensure the list of modified brushes pulled from the tree matches the order from the complete brush list to maintain draw order
+TArray<ALandmassActor*> ALandmassManagerBase::SortBrushes(TArray<ALandmassActor*> BrushArrayToMatch, TArray<ALandmassActor*> ActorsToSort)
+{
+
+	TMap<ALandmassActor*, int32> MapToSort;
+
+	for (ALandmassActor* CurrentBrush : ActorsToSort)
+	{
+		int32 IndexInSourceArray = BrushArrayToMatch.Find(CurrentBrush);
+		if(IndexInSourceArray >= 0)
+		{
+			MapToSort.Add(CurrentBrush, IndexInSourceArray);
+		}
+		
+	}
+
+	//Algo::SortBy(SortedKeys, [&SortedKeys](const AActor* A) { return  SortedKeys[A]; }, TGreater<>{});
+
+	MapToSort.ValueSort([](int32 A, int32 B) { return A < B; });
+	TArray<ALandmassActor*> SortedActors;
+	MapToSort.GenerateKeyArray(SortedActors);
+
+	return SortedActors;
+
+}
+
+void ALandmassManagerBase::AddBrushToTree(ALandmassActor* BrushToAdd, FVector4 InExtents, bool InMapToWholeLandscape, FVector4& ModifiedExtents, TArray<ALandmassActor*>& InvalidatedBrushes, TArray<int32>& ModifiedNodes)
+{
+
+	if (BrushNodeData.IsEmpty())
+	{
+		UE_LOG(LandmassManager, Warning, TEXT("Landmass Manager Brush Tree was Empty but a Brush attempted to add itself to the tree."));
+		return;
+	}
+
+	TArray<int32> RemovedFromNodes = RemoveBrushFromTree(BrushToAdd);
+	ModifiedNodes.Empty();
+
+	if (!InMapToWholeLandscape)
+	{
+		//Get the Nodes within the modified Brush Extents
+		ModifiedNodes = GetNodesWithinExtents(InExtents);
+	}
+	else
+	{
+		//If Brush affects Whole Landscape, modify the Root Node
+		ModifiedNodes = { 0 };
+	}
+
+	for (int32 ModifiedNode : ModifiedNodes)
+	{
+		BrushNodeData[ModifiedNode].BrushActors.AddUnique(BrushToAdd);
+	}
+	
+	//Update the child data counts each node contains for its children
+	UpdateChildDataCounts();
+
+	//Invalidate the currently modified Nodes and Nodes Brush was previously in
+	for (int32 RemovedNodeIndex : RemovedFromNodes)
+	{
+		ModifiedNodes.AddUnique(RemovedNodeIndex);
+	}
+
+	ModifiedExtents = FVector4(TNumericLimits<double>::Max(), TNumericLimits<double>::Max(), TNumericLimits<double>::Min(), TNumericLimits<double>::Min());
+	for (int32 ModifiedNodeIndex : ModifiedNodes)
+	{
+		FVector4 CurrentExtents = BrushNodeData[ModifiedNodeIndex].NodeExtents;
+
+		//Min of Extent Mins
+		ModifiedExtents.X = FMath::Min(CurrentExtents.X, ModifiedExtents.X);
+		ModifiedExtents.Y = FMath::Min(CurrentExtents.Y, ModifiedExtents.Y);
+		//Max of Extents Maxs
+		ModifiedExtents.Z = FMath::Max(CurrentExtents.Z, ModifiedExtents.Z);
+		ModifiedExtents.W = FMath::Max(CurrentExtents.W, ModifiedExtents.W);
+	}
+
+	//Now re-gather ModfiedNodes using the combined Modified Extents including nodes Actor removed from
+	ModifiedNodes = GetNodesWithinExtents(ModifiedExtents);
+
+	InvalidatedBrushes = GetActorsWithinModifiedNodes(ModifiedNodes);
+}
+
+void ALandmassManagerBase::AddBrushToArray(ALandmassActor* BrushToAdd)
+{
+	LandmassBrushes.AddUnique(BrushToAdd);
+	int32 LastBrushIndex = LandmassBrushes.Num() - 1;
+	for (int i = 0; i <= LastBrushIndex; i++)
+	{
+		//Remove null blueprints caused by deleting brushes
+		if (LandmassBrushes[LastBrushIndex - i] == nullptr)
+		{
+			LandmassBrushes.RemoveAt(LastBrushIndex - i);
+		}
+		//Remove REINST blueprints caused by recompiling blueprint
+		else
+		{
+			FString BrushName = LandmassBrushes[LastBrushIndex-i]->GetName();
+			if (BrushName.StartsWith("REINST"))
+			{
+				LandmassBrushes.RemoveAt(LastBrushIndex - i);
+			}
+		}
+	}
+}
+
+void ALandmassManagerBase::RequestUpdateFromBrush_Implementation(ALandmassActor* BrushRequestingUpdate)
+{
+
+}
+
+void ALandmassManagerBase::DrawBrushMaterial_Implementation(ALandmassActor* Brush, UMaterialInterface* BrushMaterial)
+{
+
+}
+
+void ALandmassManagerBase::LaunchLandmassEditor_Implementation(ALandmassActor* BrushRequestingEditor)
+{
+
+}
+
+void ALandmassManagerBase::TogglePreviewMode_Implementation(bool bEnablePreviewMode)
+{
+
+}
+
+void ALandmassManagerBase::MoveBrushUp(ALandmassActor* BrushToMove)
+{
+	int32 BrushIndex = LandmassBrushes.Find(BrushToMove);
+	if (BrushIndex >= 0 && BrushIndex < (LandmassBrushes.Num() - 1))
+	{
+		LandmassBrushes.Swap(BrushIndex, BrushIndex + 1);
+	}
+
+	RequestUpdateFromBrush(BrushToMove);
+}
+
+void ALandmassManagerBase::MoveBrushDown(ALandmassActor* BrushToMove)
+{
+	int32 BrushIndex = LandmassBrushes.Find(BrushToMove);
+	if (BrushIndex > 0)
+	{
+		LandmassBrushes.Swap(BrushIndex, BrushIndex - 1);
+	}
+
+	RequestUpdateFromBrush(BrushToMove);
+}
+
+void ALandmassManagerBase::MoveBrushToTop(ALandmassActor* BrushToMove)
+{
+	int32 BrushIndex = LandmassBrushes.Find(BrushToMove);
+	if (BrushIndex >= 0 && BrushIndex < (LandmassBrushes.Num() - 1))
+	{
+		LandmassBrushes.Remove(BrushToMove);
+		LandmassBrushes.Add(BrushToMove);
+	}
+
+	RequestUpdateFromBrush(BrushToMove);
+}
+
+void ALandmassManagerBase::MoveBrushToBottom(ALandmassActor* BrushToMove)
+{
+	int32 BrushIndex = LandmassBrushes.Find(BrushToMove);
+	if (BrushIndex > 0)
+	{
+		LandmassBrushes.Remove(BrushToMove);
+		LandmassBrushes.Insert(BrushToMove, 0);
+	}
+
+	RequestUpdateFromBrush(BrushToMove);
+}
+
+ALandscape* ALandmassManagerBase::GetLandscape()
+{
+	return DetailPanelLandscape.Get();
+}
+
+void ALandmassManagerBase::SetOwningLandscape(ALandscape* InOwningLandscape)
+{
+	Super::SetOwningLandscape(InOwningLandscape);
+
+	DetailPanelLandscape = OwningLandscape;
+}
+
+// We override PostEditChange to allow the users to change the owning landscape via a property displayed in the detail panel.
+void ALandmassManagerBase::PostEditChangeProperty(FPropertyChangedEvent& PropertyChangedEvent)
+{
+	Super::PostEditChangeProperty(PropertyChangedEvent);
+
+	// Do a bunch of checks to make sure that we don't try to do anything when the editing is happening inside the blueprint editor.
+	UWorld* World = GetWorld();
+	if (IsTemplate() || !IsValid(this) || !IsValid(World) || World->WorldType != EWorldType::Editor)
+	{
+		return;
+	}
+
+	if (PropertyChangedEvent.Property
+		&& (PropertyChangedEvent.Property->GetFName() == GET_MEMBER_NAME_CHECKED(ALandmassManagerBase, DetailPanelLandscape)))
+	{
+		SetTargetLandscape(DetailPanelLandscape.Get());
+	}
+}
+
+void ALandmassManagerBase::HandleActorSelectionChanged(const TArray<UObject*>& NewSelection, bool bForceRefresh)
+{
+	if (!IsTemplate())
+	{
+		bool bUpdateActor = false;
+		if (bWasSelected && !NewSelection.Contains(this))
+		{
+			bWasSelected = false;
+			bUpdateActor = true;
+		}
+		if (!bWasSelected && NewSelection.Contains(this))
+		{
+			bWasSelected = true;
+			bUpdateActor = true;
+		}
+		if (bUpdateActor)
+		{
+			ActorSelectionChanged(bWasSelected);
+		}
+	}
+}
+
+void ALandmassManagerBase::ActorSelectionChanged_Implementation(bool bSelected)
+{
+
+}
\ No newline at end of file
diff --git a/Engine/Plugins/Experimental/Landmass/Source/Editor/Public/LandmassActor.h b/Engine/Plugins/Experimental/Landmass/Source/Editor/Public/LandmassActor.h
--- a/Engine/Plugins/Experimental/Landmass/Source/Editor/Public/LandmassActor.h
+++ b/Engine/Plugins/Experimental/Landmass/Source/Editor/Public/LandmassActor.h
@@ -3,15 +3,36 @@
 #pragma once
 
 #include "GameFramework/Actor.h"
+#include "LandscapeBlueprintBrushBase.h"
 #include "LandmassActor.generated.h"
 
 
-UCLASS(Blueprintable, hidecategories = (Replication, Input, LOD, Actor, Cooking, Rendering))
+class ALandmassManagerBase;
+class UStaticMeshComponent;
+class UMaterialInterface;
+class UMaterialInstanceDynamic;
+
+DECLARE_DYNAMIC_MULTICAST_DELEGATE(FOnBrushUpdatedDelegate);
+
+UENUM(BlueprintType)
+enum class EBrushBlendMode : uint8
+{
+	AlphaBlend,
+	Min,
+	Max,
+	Additive
+};
+
+UCLASS(Blueprintable, PrioritizeCategories = ("PreviewMode", "LayerOrdering", "BrushSettings", "HeightmapSettings", "WeightmapSettings", "Debug"),
+	hidecategories = (Replication, Input, LOD, Actor, Cooking, Rendering, Collision, HLOD, Physics, Networking, LevelInstance, DataLayers))
 class ALandmassActor : public AActor
 {
-	GENERATED_UCLASS_BODY()
+	GENERATED_BODY()
+
+	virtual void OnConstruction(const FTransform& Transform) override;
 
 public:
+	ALandmassActor();
 
 	UFUNCTION(BlueprintNativeEvent, CallInEditor, BlueprintCallable, Category = "Tick")
 	void CustomTick(float DeltaSeconds);
@@ -21,18 +42,110 @@ public:
 	virtual bool ShouldTickIfViewportsOnly() const override;
 	virtual void Tick(float DeltaSeconds) override;
 
+	UFUNCTION(BlueprintCallable, CallInEditor, Category = Default)
+	void RenderLayer_Native(const FLandscapeBrushParameters& InParameters);
+
+	UFUNCTION(BlueprintNativeEvent, CallInEditor, meta = (ForceAsFunction) , Category = Default)
+	void RenderLayer(const FLandscapeBrushParameters& InParameters);
+
+	UFUNCTION(BlueprintCallable, CallInEditor, Category = "PreviewMode", meta = (DisplayPriority = 1))
+	void FastPreviewMode();
+
+	UFUNCTION(BlueprintCallable, CallInEditor, Category = "PreviewMode", meta = (DisplayPriority = 2))
+	void RestoreLandscapeEditing();
+
+	UFUNCTION(BlueprintCallable, CallInEditor, Category = "LayerOrdering", meta = (DisplayPriority = 3))
+	void MoveBrushUp();
+
+	UFUNCTION(BlueprintCallable, CallInEditor, Category = "LayerOrdering", meta = (DisplayPriority = 4))
+	void MoveBrushDown();
+
+	UFUNCTION(BlueprintCallable, CallInEditor, Category = "LayerOrdering", meta = (DisplayPriority = 5))
+	void MoveToTop();
+
+	UFUNCTION(BlueprintCallable, CallInEditor, Category = "LayerOrdering", meta = (DisplayPriority = 6))
+	void MoveToBottom();
+	
+	UPROPERTY(EditAnywhere, BlueprintReadWrite, Category = "BrushSettings", meta = (DisplayPriority = 8))
+	float BrushSize = 4096;
+
+	UPROPERTY(EditAnywhere, BlueprintReadWrite, Category = "BrushSettings", meta = (DisplayPriority = 9))
+	bool DrawToEntireLandscape = false;
+
+	UPROPERTY(EditAnywhere, BlueprintReadWrite, Category = "BrushSettings", meta = (DisplayPriority = 10))
+	bool AffectsHeightmap = true;
+
+	UPROPERTY(EditAnywhere, BlueprintReadWrite, Category = "BrushSettings", meta = (DisplayPriority = 11))
+	bool AffectsWeightmaps = true;
+
+	UPROPERTY(EditAnywhere, BlueprintReadWrite, Category = "BrushSettings", meta = (DisplayPriority = 12))
+	bool AffectsVisibility = false;
+	
+	UPROPERTY(EditAnywhere, BlueprintReadWrite, Category = "HeightmapSettings", meta = (DisplayPriority = 13))
+	EBrushBlendMode HeightBlendMode = EBrushBlendMode::AlphaBlend;
+
+	UPROPERTY(EditAnywhere, BlueprintReadWrite, Category = "HeightmapSettings", meta = (DisplayPriority = 14))
+	TObjectPtr<UMaterialInterface> HeightMaterial;
+
+	UPROPERTY(EditAnywhere, BlueprintReadWrite, Category = "WeightmapSettings", meta = (DisplayPriority = 15))
+	EBrushBlendMode WeightMapBlendMode = EBrushBlendMode::AlphaBlend;
+
+	UPROPERTY(EditAnywhere, BlueprintReadWrite, Category = "WeightmapSettings", meta = (DisplayPriority = 16))
+	TObjectPtr<UMaterialInterface> WeightmapMaterial;
+
+	UPROPERTY(BlueprintReadOnly, Category = "Debug", meta = (DisplayPriority = 17))
+	FVector4 BrushExtents;
+
+	UPROPERTY(EditAnywhere, BlueprintReadWrite, Category = "Debug", meta = (DisplayPriority = 18))
+	TArray<FName> WeightmapLayers;
+
+	UPROPERTY(EditAnywhere, BlueprintReadWrite, Category = "Debug", meta = (DisplayPriority = 19))
+	FLandscapeBrushParameters BrushRenderParameters;
+
+	UPROPERTY(EditAnywhere, BlueprintReadWrite, Category = "Debug", meta = (DisplayPriority = 20))
+	TObjectPtr<ALandmassManagerBase> BrushManager;
+
 	UFUNCTION(BlueprintCallable, category = "Default")
 	void SetEditorTickEnabled(bool bEnabled) { EditorTickIsEnabled = bEnabled; }
 
+	UFUNCTION(BlueprintCallable, Category = "Landmass")
+	void UpdateBrushExtents();
+
 	UPROPERTY()
 	bool EditorTickIsEnabled = false;
 
 	UFUNCTION(BlueprintNativeEvent, CallInEditor, BlueprintCallable, Category = "Selection")
 	void ActorSelectionChanged(bool bSelected);
 
+	void FindOrSpawnManager();
+
+	UFUNCTION(BlueprintCallable, CallInEditor, Category = Default)
+	void DrawBrushMaterial(UMaterialInterface* Material);
+
+	UFUNCTION(BlueprintCallable, CallInEditor, Category = Default)
+	void SetMeshExentsMaterial(UMaterialInterface* Material);
+
+	UPROPERTY(BlueprintAssignable, Category = "Landmass")
+	FOnBrushUpdatedDelegate OnBrushUpdated;
+
+	UPROPERTY(EditAnywhere, BlueprintReadWrite, Transient, Category = "Debug", meta = (DisplayPriority = 21))
+	TObjectPtr<UMaterialInstanceDynamic> ExtentsPreviewMID;
+
+	UPROPERTY(EditAnywhere, BlueprintReadWrite, Transient, Category = "Debug", meta = (DisplayPriority = 22))
+	TObjectPtr<UMaterialInstanceDynamic> HeightmapRenderMID;
+
+	UPROPERTY(EditAnywhere, BlueprintReadWrite, Transient, Category = "Debug", meta = (DisplayPriority = 23))
+	TObjectPtr<UMaterialInstanceDynamic> WeightmapRenderMID;
+
 private:
 	bool bWasSelected = false;
 
+	UPROPERTY()
+	TObjectPtr<UStaticMeshComponent> MeshExtentsQuad;
+
+	UPROPERTY()
+	TObjectPtr<UStaticMeshComponent> BrushSpriteMesh;
+
 	FDelegateHandle OnActorSelectionChangedHandle;
 
 	/** Called when the editor selection has changed. */
diff --git a/Engine/Plugins/Experimental/Landmass/Source/Editor/Public/LandmassBPEditorExtension.h b/Engine/Plugins/Experimental/Landmass/Source/Editor/Public/LandmassBPEditorExtension.h
--- a/Engine/Plugins/Experimental/Landmass/Source/Editor/Public/LandmassBPEditorExtension.h
+++ b/Engine/Plugins/Experimental/Landmass/Source/Editor/Public/LandmassBPEditorExtension.h
@@ -3,16 +3,28 @@
 #pragma once
 
 #include "Kismet/BlueprintFunctionLibrary.h"
+#include "LandmassManagerBase.h"
 #include "LandmassBPEditorExtension.generated.h"
 
 
-UCLASS()
+UCLASS(meta = (Namespace = "Landmass"))
 class ULandmassBlueprintFunctionLibrary : public UBlueprintFunctionLibrary
 {
-	GENERATED_UCLASS_BODY()
+	GENERATED_BODY()
+
+public:
 
 	UFUNCTION(BlueprintCallable, meta = (DisplayName = "Get Cursor World Ray", Keywords = "Get Cursor World Ray", UnsafeDuringActorConstruction = "true"), Category = Rendering)
 	static bool GetCursorWorldRay(FVector& CameraLocation, FVector& RayOrigin, FVector& RayDirection);
+
+	UFUNCTION(BlueprintCallable, BlueprintPure, meta = (DisplayName = "Get Overlapping World Extents", Keywords = "Combine World Extents", UnsafeDuringActorConstruction = "false"), Category = Landmass)
+	static void CombineWorldExtents(FVector4 ExtentsA, FVector4 ExtentsB, FVector4& CombinedExtents);
+
+	UFUNCTION(BlueprintCallable, BlueprintPure, meta = (DisplayName = "World Extents to Landmass Coordinates", Keywords = "World Extents to Landmass Coordinates", UnsafeDuringActorConstruction = "false"), Category = Landmass)
+	static void WorldExtentsToCanvasCoordinates(FVector4 WorldExtents, FLandmassLandscapeInfo LandscapeInfo, FVector2D& ScreenPosition, FVector2D& ScreenSize, FVector2D& CoordinatePosition, FVector2D& CoordinateSize);
+
+	UFUNCTION(BlueprintCallable, Category = Rendering)
+	static void ForceUpdateTexture(UTexture* InTexture);
 };
 
 #if UE_ENABLE_INCLUDE_ORDER_DEPRECATED_IN_5_2
diff --git a/Engine/Plugins/Experimental/Landmass/Source/Editor/Public/LandmassErosionBrushBase.h b/Engine/Plugins/Experimental/Landmass/Source/Editor/Public/LandmassErosionBrushBase.h
new file mode 100644
--- /dev/null
+++ b/Engine/Plugins/Experimental/Landmass/Source/Editor/Public/LandmassErosionBrushBase.h
@@ -0,0 +1,67 @@
+// Copyright Epic Games, Inc. All Rights Reserved.
+
+#pragma once
+
+#include "LandscapeBlueprintBrushBase.h"
+#include "LandscapeEditTypes.h"
+#include "Logging/LogMacros.h"
+#include "LandmassErosionBrushBase.generated.h"
+
+LANDMASSEDITOR_API DECLARE_LOG_CATEGORY_EXTERN(LandmassErosionBrush, Display, All);
+
+UCLASS(Blueprintable, hidecategories = (Replication, Input, LOD, Actor, Cooking, Rendering))
+class ALandmassErosionBrushBase : public ALandscapeBlueprintBrushBase
+{
+	GENERATED_BODY()
+
+	virtual void OnConstruction(const FTransform& Transform) override;
+
+public:
+
+	ALandmassErosionBrushBase();
+
+	virtual bool IsEditorOnly() const override { return true; }
+
+	// Adds the brush to the given landscape, removing it from any previous one. This differs from SetOwningLandscape
+	// in that SetOwningLandscape is called by the landscape itself from AddBrushToLayer to update the manager.
+	UFUNCTION(BlueprintCallable, category = "Landmass")
+	virtual void SetTargetLandscape(ALandscape* InOwningLandscape);
+
+	UFUNCTION(BlueprintCallable, category = "Landmass")
+	ALandscape* GetLandscape();
+
+	UFUNCTION(BlueprintNativeEvent, CallInEditor, BlueprintCallable, Category = "Selection")
+	void ActorSelectionChanged(bool bSelected);
+
+	void FindAndAssignLandscape();
+
+#if WITH_EDITOR
+	// ALandscapeBlueprintBrushBase
+	virtual void SetOwningLandscape(class ALandscape* InOwningLandscape) override;
+	virtual void PostEditChangeProperty(FPropertyChangedEvent& PropertyChangedEvent) override;
+#endif
+
+private:
+	bool bWasSelected = false;
+
+	FDelegateHandle OnActorSelectionChangedHandle;
+
+	/** Called when the editor selection has changed. */
+	void HandleActorSelectionChanged(const TArray<UObject*>& NewSelection, bool bForceRefresh);
+
+#if WITH_EDITORONLY_DATA
+	//~ This is transient because SetOwningLandscape is called in ALandscape::PostLoad.
+	/**
+	 * The owning landscape.
+	 */
+	UPROPERTY(EditAnywhere, Category = Landscape, Transient, meta = (DisplayName = "Landscape"))
+	TObjectPtr<ALandscape> DetailPanelLandscape;
+#endif
+
+};
+
+#if UE_ENABLE_INCLUDE_ORDER_DEPRECATED_IN_5_2
+#include "CoreMinimal.h"
+#include "Tickable.h"
+#include "UObject/GCObject.h"
+#endif
diff --git a/Engine/Plugins/Experimental/Landmass/Source/Editor/Public/LandmassManagerBase.h b/Engine/Plugins/Experimental/Landmass/Source/Editor/Public/LandmassManagerBase.h
new file mode 100644
--- /dev/null
+++ b/Engine/Plugins/Experimental/Landmass/Source/Editor/Public/LandmassManagerBase.h
@@ -0,0 +1,167 @@
+// Copyright Epic Games, Inc. All Rights Reserved.
+
+#pragma once
+
+#include "LandscapeBlueprintBrushBase.h"
+#include "LandmassActor.h"
+#include "LandscapeEditTypes.h"
+#include "Logging/LogMacros.h"
+#include "LandmassManagerBase.generated.h"
+
+LANDMASSEDITOR_API DECLARE_LOG_CATEGORY_EXTERN(LandmassManager, Display, All);
+
+USTRUCT(BlueprintType)
+struct FBrushDataTree
+{
+	GENERATED_BODY()
+	FBrushDataTree() {
+		CurrentLevel = -1; ParentIndex = -1; Index_x0y0 = -1; Index_x1y0 = -1; Index_x0y1 = -1; Index_x1y1 = -1; ChildDataCount = 0; NodeExtents = FVector4(0, 0, 0, 0);
+	}
+
+	UPROPERTY(EditAnywhere, BlueprintReadWrite, Category = Landmass)
+	int32 CurrentLevel;
+	UPROPERTY(EditAnywhere, BlueprintReadWrite, Category = Landmass)
+	int32 ParentIndex;
+	UPROPERTY(EditAnywhere, BlueprintReadWrite, Category = Landmass)
+	TArray<TObjectPtr<ALandmassActor>> BrushActors;
+	UPROPERTY(EditAnywhere, BlueprintReadWrite, Category = Landmass)
+	int32 Index_x0y0;
+	UPROPERTY(EditAnywhere, BlueprintReadWrite, Category = Landmass)
+	int32 Index_x1y0;
+	UPROPERTY(EditAnywhere, BlueprintReadWrite, Category = Landmass)
+	int32 Index_x0y1;
+	UPROPERTY(EditAnywhere, BlueprintReadWrite, Category = Landmass)
+	int32 Index_x1y1;
+	UPROPERTY(EditAnywhere, BlueprintReadWrite, Category = Landmass)
+	int32 ChildDataCount;
+	UPROPERTY(EditAnywhere, BlueprintReadWrite, Category = Landmass)
+	FVector4 NodeExtents;
+};
+
+USTRUCT(BlueprintType)
+struct FLandmassLandscapeInfo
+{
+	GENERATED_BODY()
+	FLandmassLandscapeInfo() {
+		LandscapeTransform = FTransform(); LandscapeQuads = { 0,0 }; RenderTargetResolution = { 0,0 };
+	}
+
+	UPROPERTY(EditAnywhere, BlueprintReadWrite, Category = Landmass)
+	FTransform LandscapeTransform;
+	UPROPERTY(EditAnywhere, BlueprintReadWrite, Category = Landmass)
+	FIntPoint LandscapeQuads;
+	UPROPERTY(EditAnywhere, BlueprintReadWrite, Category = Landmass)
+	FIntPoint RenderTargetResolution;
+
+};
+
+UCLASS(Blueprintable, hidecategories = (Replication, Input, LOD, Actor, Cooking, Rendering))
+class ALandmassManagerBase : public ALandscapeBlueprintBrushBase
+{
+	GENERATED_BODY()
+
+public:
+
+	ALandmassManagerBase();
+
+	virtual bool IsEditorOnly() const override { return true; }
+
+	UPROPERTY(EditAnywhere, BlueprintReadWrite, Category = Quadtree)
+	TArray<FBrushDataTree> BrushNodeData;
+
+	UPROPERTY(EditAnywhere, BlueprintReadWrite, Category = Landmass)
+	FLandmassLandscapeInfo LandscapeInformation;
+
+	UPROPERTY(EditAnywhere, BlueprintReadWrite, Category = Quadtree)
+	int32 BrushTreeDepth;
+
+	UPROPERTY(EditAnywhere, BlueprintReadWrite, Category = Landmass)
+	TArray<TObjectPtr<ALandmassActor>> LandmassBrushes;
+
+	// Adds the brush to the given landscape, removing it from any previous one. This differs from SetOwningLandscape
+	// in that SetOwningLandscape is called by the landscape itself from AddBrushToLayer to update the manager.
+	UFUNCTION(BlueprintCallable, category = "Landmass")
+	virtual void SetTargetLandscape(ALandscape* InOwningLandscape);
+
+	int32 Convert2DTo1D(const FIntPoint& Index2D, int32 XSize);
+
+	UFUNCTION(BlueprintCallable, category = "Landmass")
+	void PopulateNodeTree();
+
+	UFUNCTION(BlueprintCallable, category = "Landmass")
+	TArray<ALandmassActor*> GetActorsWithinModifiedNodes(TArray<int32>& InModifiedNodes);
+
+	UFUNCTION(BlueprintCallable, category = "Landmass")
+	void UpdateChildDataCounts();
+
+	UFUNCTION(BlueprintCallable, category = "Landmass")
+	void ConsolidateNodes(TArray<int32>& NodesToConsolidate);
+
+	UFUNCTION(BlueprintCallable, category = "Landmass")
+	TArray<int32> GetNodesWithinExtents(FVector4& InExtents);
+
+	UFUNCTION(BlueprintCallable, category = "Landmass")
+	TArray<int32> RemoveBrushFromTree(ALandmassActor* BrushToRemove);
+
+	UFUNCTION(BlueprintCallable, category = "Landmass")
+	TArray<ALandmassActor*> SortBrushes(TArray<ALandmassActor*> BrushArrayToMatch, TArray<ALandmassActor*> ActorsToSort);
+
+	UFUNCTION(BlueprintCallable, category = "Landmass")
+	void AddBrushToTree(ALandmassActor* BrushToAdd, FVector4 InExtents, bool InMapToWholeLandscape, FVector4& ModifiedExtents, TArray<ALandmassActor*>& InvalidatedBrushes, TArray<int32>& ModifiedNodes);
+
+	UFUNCTION(BlueprintCallable, category = "Landmass")
+	void AddBrushToArray(ALandmassActor* BrushToAdd);
+
+	UFUNCTION(BlueprintNativeEvent, BlueprintCallable, meta = (ForceAsFunction), Category = Default)
+	void RequestUpdateFromBrush(ALandmassActor* BrushRequestingUpdate);
+
+	UFUNCTION(BlueprintNativeEvent, BlueprintCallable, CallInEditor, meta = (ForceAsFunction, DefaultToSelf = "Brush"), Category = Default)
+	void DrawBrushMaterial(ALandmassActor* Brush, UMaterialInterface* BrushMaterial);
+
+	UFUNCTION(BlueprintNativeEvent, BlueprintCallable, meta = (ForceAsFunction), Category = Default)
+	void LaunchLandmassEditor(ALandmassActor* BrushRequestingEditor);
+
+	UFUNCTION(BlueprintNativeEvent, BlueprintCallable, meta = (ForceAsFunction), Category = Default)
+	void TogglePreviewMode(bool bEnablePreviewMode);
+	
+	void MoveBrushUp(ALandmassActor* BrushToMove);
+	void MoveBrushDown(ALandmassActor* BrushToMove);
+	void MoveBrushToTop(ALandmassActor* BrushToMove);
+	void MoveBrushToBottom(ALandmassActor* BrushToMove);
+
+	UFUNCTION(BlueprintCallable, category = "Landmass")
+	ALandscape* GetLandscape();
+
+	UFUNCTION(BlueprintNativeEvent, CallInEditor, BlueprintCallable, Category = "Selection")
+	void ActorSelectionChanged(bool bSelected);
+
+#if WITH_EDITOR
+	// ALandscapeBlueprintBrushBase
+	virtual void SetOwningLandscape(class ALandscape* InOwningLandscape) override;
+	virtual void PostEditChangeProperty(FPropertyChangedEvent& PropertyChangedEvent) override;
+#endif
+
+private:
+	bool bWasSelected = false;
+
+	FDelegateHandle OnActorSelectionChangedHandle;
+
+	/** Called when the editor selection has changed. */
+	void HandleActorSelectionChanged(const TArray<UObject*>& NewSelection, bool bForceRefresh);
+
+#if WITH_EDITORONLY_DATA
+	//~ This is transient because SetOwningLandscape is called in ALandscape::PostLoad.
+	/**
+	 * The owning landscape.
+	 */
+	UPROPERTY(EditAnywhere, Category = Landscape, Transient, meta = (DisplayName = "Landscape"))
+	TObjectPtr<ALandscape> DetailPanelLandscape;
+#endif
+
+};
+
+#if UE_ENABLE_INCLUDE_ORDER_DEPRECATED_IN_5_2
+#include "CoreMinimal.h"
+#include "Tickable.h"
+#include "UObject/GCObject.h"
+#endif
