LandmassActor and LandmassErosionBrushBase: Added check to make sure is editor world before finding and assingning landscape

diff --git a/Engine/Plugins/Experimental/Landmass/Source/Editor/Private/LandmassActor.cpp b/Engine/Plugins/Experimental/Landmass/Source/Editor/Private/LandmassActor.cpp
--- a/Engine/Plugins/Experimental/Landmass/Source/Editor/Private/LandmassActor.cpp
+++ b/Engine/Plugins/Experimental/Landmass/Source/Editor/Private/LandmassActor.cpp
@@ -156,6 +156,11 @@ void ALandmassActor::FindOrSpawnManager()
 		return;
 	}
 
+	if (GetWorld()->WorldType != EWorldType::Editor)
+	{
+		return;
+	}
+
 	TArray<AActor*> FoundActors;
 	UGameplayStatics::GetAllActorsOfClass(Cast<UObject>(GetWorld()), ALandmassManagerBase::StaticClass(), FoundActors);
 
diff --git a/Engine/Plugins/Experimental/Landmass/Source/Editor/Private/LandmassErosionBrushBase.cpp b/Engine/Plugins/Experimental/Landmass/Source/Editor/Private/LandmassErosionBrushBase.cpp
--- a/Engine/Plugins/Experimental/Landmass/Source/Editor/Private/LandmassErosionBrushBase.cpp
+++ b/Engine/Plugins/Experimental/Landmass/Source/Editor/Private/LandmassErosionBrushBase.cpp
@@ -32,6 +32,12 @@ void ALandmassErosionBrushBase::OnConstruction(const FTransform& Transform)
 void ALandmassErosionBrushBase::SetTargetLandscape(ALandscape* InTargetLandscape)
 {
 #if WITH_EDITOR
+
+	if (GetWorld()->WorldType != EWorldType::Editor)
+	{
+		return;
+	}
+
 	if (OwningLandscape != InTargetLandscape)
 	{
 		if (OwningLandscape)
