LandmassActor.h: Temporarily disabling spawning landmass manager automatically due to build errors.

diff --git a/Engine/Plugins/Experimental/Landmass/Source/Editor/Private/LandmassActor.cpp b/Engine/Plugins/Experimental/Landmass/Source/Editor/Private/LandmassActor.cpp
--- a/Engine/Plugins/Experimental/Landmass/Source/Editor/Private/LandmassActor.cpp
+++ b/Engine/Plugins/Experimental/Landmass/Source/Editor/Private/LandmassActor.cpp
@@ -47,7 +47,7 @@ void ALandmassActor::OnConstruction(const FTransform& Transform)
 {
 	Super::OnConstruction(Transform);
 	UpdateBrushExtents();
-	FindOrSpawnManager();
+	//FindOrSpawnManager();
 	if (BrushManager != nullptr)
 	{
 		if (AffectsHeightmap || AffectsWeightmaps || AffectsVisibility)
