Landmass: Added option to compute seamless normals for connected landscape tiles to LandmassManagerBase.h.

diff --git a/Engine/Plugins/Experimental/Landmass/Source/Editor/Private/LandmassManagerBase.cpp b/Engine/Plugins/Experimental/Landmass/Source/Editor/Private/LandmassManagerBase.cpp
--- a/Engine/Plugins/Experimental/Landmass/Source/Editor/Private/LandmassManagerBase.cpp
+++ b/Engine/Plugins/Experimental/Landmass/Source/Editor/Private/LandmassManagerBase.cpp
@@ -492,6 +492,11 @@ void ALandmassManagerBase::TogglePreviewMode_Implementation(bool bEnablePreviewM
 
 }
 
+void ALandmassManagerBase::SetCaptureBoundaryNormals(bool bInCaptureBoundaryNormals)
+{
+	bCaptureBoundaryNormals = bInCaptureBoundaryNormals;
+}
+
 void ALandmassManagerBase::MoveBrushUp(ALandmassActor* BrushToMove)
 {
 	int32 BrushIndex = LandmassBrushes.Find(BrushToMove);
diff --git a/Engine/Plugins/Experimental/Landmass/Source/Editor/Public/LandmassManagerBase.h b/Engine/Plugins/Experimental/Landmass/Source/Editor/Public/LandmassManagerBase.h
--- a/Engine/Plugins/Experimental/Landmass/Source/Editor/Public/LandmassManagerBase.h
+++ b/Engine/Plugins/Experimental/Landmass/Source/Editor/Public/LandmassManagerBase.h
@@ -123,6 +123,9 @@ public:
 
 	UFUNCTION(BlueprintNativeEvent, BlueprintCallable, meta = (ForceAsFunction), Category = Default)
 	void TogglePreviewMode(bool bEnablePreviewMode);
+
+	UFUNCTION(BlueprintCallable, Category = Default)
+	void SetCaptureBoundaryNormals(bool bInCaptureBoundaryNormals);
 	
 	void MoveBrushUp(ALandmassActor* BrushToMove);
 	void MoveBrushDown(ALandmassActor* BrushToMove);
