Landmass Actor: Fix brush to not have the bounds component visible when spawned procedurally.

diff --git a/Engine/Plugins/Experimental/Landmass/Source/Editor/Private/LandmassActor.cpp b/Engine/Plugins/Experimental/Landmass/Source/Editor/Private/LandmassActor.cpp
--- a/Engine/Plugins/Experimental/Landmass/Source/Editor/Private/LandmassActor.cpp
+++ b/Engine/Plugins/Experimental/Landmass/Source/Editor/Private/LandmassActor.cpp
@@ -48,6 +48,7 @@ void ALandmassActor::OnConstruction(const FTransform& Transform)
 	Super::OnConstruction(Transform);
 	UpdateBrushExtents();
 	FindOrSpawnManager();
+	MeshExtentsQuad->SetVisibility(false);
 	if (BrushManager != nullptr)
 	{
 		if (AffectsHeightmap || AffectsWeightmaps || AffectsVisibility)
