Landmass: LandmassActor: Fix so preview extent mesh correctly shows up when the actors are selected and while updating settings.

diff --git a/Engine/Plugins/Experimental/Landmass/Source/Editor/Private/LandmassActor.cpp b/Engine/Plugins/Experimental/Landmass/Source/Editor/Private/LandmassActor.cpp
--- a/Engine/Plugins/Experimental/Landmass/Source/Editor/Private/LandmassActor.cpp
+++ b/Engine/Plugins/Experimental/Landmass/Source/Editor/Private/LandmassActor.cpp
@@ -48,7 +48,7 @@ void ALandmassActor::OnConstruction(const FTransform& Transform)
 	Super::OnConstruction(Transform);
 	UpdateBrushExtents();
 	FindOrSpawnManager();
-	MeshExtentsQuad->SetVisibility(false);
+	MeshExtentsQuad->SetVisibility(bWasSelected);
 	if (BrushManager != nullptr)
 	{
 		if (AffectsHeightmap || AffectsWeightmaps || AffectsVisibility)
