Fix blutility module not guaranteed to be loaded when Landmass engine plugin compiles its content depending on that module

#rnx

diff --git a/Engine/Plugins/Experimental/Landmass/Source/Editor/Private/LandmassEditorModule.cpp b/Engine/Plugins/Experimental/Landmass/Source/Editor/Private/LandmassEditorModule.cpp
--- a/Engine/Plugins/Experimental/Landmass/Source/Editor/Private/LandmassEditorModule.cpp
+++ b/Engine/Plugins/Experimental/Landmass/Source/Editor/Private/LandmassEditorModule.cpp
@@ -9,6 +9,10 @@ class FLandmassEditorModule : public ILandmassEditorModuleInterface
 public:
 	virtual void StartupModule() override
 	{
+		// If a plugin's *content* directory contains Blutilities (blueprints inheriting classes from the Blutility engine module), there is no
+		// guarantee that the Blutility engine module is loaded when the plugin gets loaded. The following is a workaround to enforce that the
+		// required parent classes are loaded before we try to compile this plugin's content.
+		FModuleManager::Get().LoadModuleChecked("Blutility");
 	}
 
 	virtual void ShutdownModule() override
